[AutoRTFM] Fix DLLEXPORT / DLLIMPORT duplicate definion for non-windows / non-linux platforms.

diff --git a/Engine/Source/Runtime/AutoRTFM/Public/AutoRTFM.h b/Engine/Source/Runtime/AutoRTFM/Public/AutoRTFM.h
--- a/Engine/Source/Runtime/AutoRTFM/Public/AutoRTFM.h
+++ b/Engine/Source/Runtime/AutoRTFM/Public/AutoRTFM.h
@@ -82,13 +82,15 @@
 		#ifndef DLLIMPORT
 			#define DLLIMPORT __declspec(dllimport)
 		#endif
-	#else
+	#elif defined(PLATFORM_LINUX) && PLATFORM_LINUX
 		#ifndef DLLEXPORT
 			#define DLLEXPORT __attribute__((visibility("default")))
 		#endif
 		#ifndef DLLIMPORT
 			#define DLLIMPORT __attribute__((visibility("default")))
 		#endif
+	#else
+		// AutoRTFM is not enabled on any other platform
 	#endif
 #endif
 
