[Backout] - CL40940143
#rnx
Original CL Desc
-----------------------------------------------------------------
[AutoRTFM] Fix DLLEXPORT / DLLIMPORT duplicate definion for non-windows / non-linux platforms.

diff --git a/Engine/Source/Runtime/AutoRTFM/Public/AutoRTFM.h b/Engine/Source/Runtime/AutoRTFM/Public/AutoRTFM.h
--- a/Engine/Source/Runtime/AutoRTFM/Public/AutoRTFM.h
+++ b/Engine/Source/Runtime/AutoRTFM/Public/AutoRTFM.h
@@ -82,15 +82,13 @@
 		#ifndef DLLIMPORT
 			#define DLLIMPORT __declspec(dllimport)
 		#endif
-	#elif defined(PLATFORM_LINUX) && PLATFORM_LINUX
+	#else
 		#ifndef DLLEXPORT
 			#define DLLEXPORT __attribute__((visibility("default")))
 		#endif
 		#ifndef DLLIMPORT
 			#define DLLIMPORT __attribute__((visibility("default")))
 		#endif
-	#else
-		// AutoRTFM is not enabled on any other platform
 	#endif
 #endif
 
