[AutoRTFM] Replace includes of "AutoRTFM/AutoRTFMMetrics.h" with "AutoRTFMMetrics.h"

Delete the now unused "AutoRTFM/AutoRTFMMetrics.h".
Part of the clean-up post move of AutoRTFM from Core to its own module.

diff --git a/Engine/Source/Runtime/AutoRTFM/Private/Context.cpp b/Engine/Source/Runtime/AutoRTFM/Private/Context.cpp
--- a/Engine/Source/Runtime/AutoRTFM/Private/Context.cpp
+++ b/Engine/Source/Runtime/AutoRTFM/Private/Context.cpp
@@ -5,7 +5,7 @@
 #include "Context.h"
 #include "ContextInlines.h"
 
-#include "AutoRTFM/AutoRTFMMetrics.h"
+#include "AutoRTFMMetrics.h"
 #include "CallNestInlines.h"
 #include "ExternAPI.h"
 #include "FunctionMap.h"
diff --git a/Engine/Source/Runtime/AutoRTFM/Public/AutoRTFM/AutoRTFMMetrics.h b/Engine/Source/Runtime/AutoRTFM/Public/AutoRTFM/AutoRTFMMetrics.h
deleted file mode 100644
