[AutoRTFM] Replace includes of "AutoRTFM/AutoRTFMConstants.h" with "AutoRTFMConstants.h"

Delete the now unused "AutoRTFM/AutoRTFMConstants.h".
Part of the clean-up post move of AutoRTFM from Core to its own module.

diff --git a/Engine/Source/Runtime/AutoRTFM/Private/API.cpp b/Engine/Source/Runtime/AutoRTFM/Private/API.cpp
--- a/Engine/Source/Runtime/AutoRTFM/Private/API.cpp
+++ b/Engine/Source/Runtime/AutoRTFM/Private/API.cpp
@@ -390,7 +390,7 @@ namespace AutoRTFM
 }
 
 #if (defined(__AUTORTFM_ENABLED) && __AUTORTFM_ENABLED)
-#include "AutoRTFM/AutoRTFMConstants.h"
+#include "AutoRTFMConstants.h"
 #include "CallNest.h"
 #include "Context.h"
 #include "ContextInlines.h"
diff --git a/Engine/Source/Runtime/AutoRTFM/Private/FunctionMapInlines.h b/Engine/Source/Runtime/AutoRTFM/Private/FunctionMapInlines.h
--- a/Engine/Source/Runtime/AutoRTFM/Private/FunctionMapInlines.h
+++ b/Engine/Source/Runtime/AutoRTFM/Private/FunctionMapInlines.h
@@ -4,7 +4,7 @@
 
 #if (defined(__AUTORTFM) && __AUTORTFM)
 
-#include "AutoRTFM/AutoRTFMConstants.h"
+#include "AutoRTFMConstants.h"
 #include "FunctionMap.h"
 #include "Utils.h"
 
diff --git a/Engine/Source/Runtime/AutoRTFM/Public/AutoRTFM/AutoRTFMConstants.h b/Engine/Source/Runtime/AutoRTFM/Public/AutoRTFM/AutoRTFMConstants.h
deleted file mode 100644
