[AutoRTFM] Add `__asan_poison_memory_region` to the closed -> open stdlib mappings

Without this, we'll crash soon after starting the ASAN+linux+test server

diff --git a/Engine/Source/Runtime/AutoRTFM/Private/StdLibSupport.cpp b/Engine/Source/Runtime/AutoRTFM/Private/StdLibSupport.cpp
--- a/Engine/Source/Runtime/AutoRTFM/Private/StdLibSupport.cpp
+++ b/Engine/Source/Runtime/AutoRTFM/Private/StdLibSupport.cpp
@@ -603,6 +603,7 @@ UE_AUTORTFM_REGISTER_OPEN_TO_CLOSED_FUNCTIONS(
 #if AUTORTFM_ASAN_ENABLED
 	UE_AUTORTFM_MAP_OPEN_TO_SELF(__asan_addr_is_in_fake_stack),
 	UE_AUTORTFM_MAP_OPEN_TO_SELF(__asan_get_current_fake_stack),
+	UE_AUTORTFM_MAP_OPEN_TO_SELF(__asan_poison_memory_region),
 #endif
 
 #if AUTORTFM_PLATFORM_WINDOWS
