[AutoRTFM] Rename our magic constant value to MagicPrefix for clarity.

Also improves the surrounding comments to give a little more context to a new reader.

diff --git a/Engine/Source/Runtime/AutoRTFM/Private/FunctionMapInlines.h b/Engine/Source/Runtime/AutoRTFM/Private/FunctionMapInlines.h
--- a/Engine/Source/Runtime/AutoRTFM/Private/FunctionMapInlines.h
+++ b/Engine/Source/Runtime/AutoRTFM/Private/FunctionMapInlines.h
@@ -28,26 +28,26 @@ inline void* FunctionMapLookup(void* OldFunction, const char* Where)
 {
 	// We use prefix data in our custom LLVM pass to stuff some data just
 	// before the address of all open function pointers (that we have
-	// definitions for!). We use the special Magic Mike constant in the top
-	// 16-bits of the function pointer address as a magic constant check
+	// definitions for!). We verify the special Magic Prefix constant in the
+	// top 16-bits of the function pointer address as a magic constant check
 	// to give us a much higher confidence that there is actually a closed
 	// variant pointer residing 8-bytes before our function address.
 	uint64_t PrefixData;
 	memcpy(&PrefixData, reinterpret_cast<char*>(OldFunction) - sizeof(uint64_t), sizeof(uint64_t));
-	if (AUTORTFM_LIKELY(Constants::MagicMike == (PrefixData & 0xffff000000000000)))
+	if (AUTORTFM_LIKELY(Constants::MagicPrefix == (PrefixData & 0xffff'0000'0000'0000)))
 	{
-		return reinterpret_cast<void*>(PrefixData & 0x0000ffffffffffff);
+		return reinterpret_cast<void*>(PrefixData & 0x0000'ffff'ffff'ffff);
 	}
 	// UBSAN adds a type hash prefix to the function as a "prologue" that
-	// ends preceding our Magic Mike prefix. They use 0xc105cafe in the
+	// ends preceding our Magic Prefix. They use 0xc105cafe in the
 	// lower 32-bits to distinguish the 64-bit word containing their type
 	// hash. If we see it, check the preceding 64-bit word for our prefix.
-	else if (AUTORTFM_UNLIKELY(0xc105cafe == (PrefixData & 0x00000000ffffffff)))
+	else if (AUTORTFM_UNLIKELY(0xc105'cafe == (PrefixData & 0x0000'0000'ffff'ffff)))
 	{
 		memcpy(&PrefixData, reinterpret_cast<char*>(OldFunction) - sizeof(uint64_t) * 2, sizeof(uint64_t));
-		if (AUTORTFM_LIKELY(Constants::MagicMike == (PrefixData & 0xffff000000000000)))
+		if (AUTORTFM_LIKELY(Constants::MagicPrefix == (PrefixData & 0xffff'0000'0000'0000)))
 		{
-			return reinterpret_cast<void*>(PrefixData & 0x0000ffffffffffff);
+			return reinterpret_cast<void*>(PrefixData & 0x0000'ffff'ffff'ffff);
 		}
 	}
 
diff --git a/Engine/Source/Runtime/AutoRTFM/Public/AutoRTFMConstants.h b/Engine/Source/Runtime/AutoRTFM/Public/AutoRTFMConstants.h
--- a/Engine/Source/Runtime/AutoRTFM/Public/AutoRTFMConstants.h
+++ b/Engine/Source/Runtime/AutoRTFM/Public/AutoRTFMConstants.h
@@ -41,10 +41,13 @@ namespace AutoRTFM::Constants
 	inline constexpr uint32_t Minor = 2;
 	inline constexpr uint32_t Patch = 0;
 
-	// The Magic Mike constant - shared with the compiler implementation.
-	// Used to ensure that prefix data pointers are what they claim to be
-	// and were definitely injected by us.
-	inline constexpr uint64_t MagicMike = 0xa273000000000000;
+	// The Magic Prefix constant - an arbitrarily chosen address prefix, shared
+	// between the AutoRTFM compiler and runtime.
+	// We add this prefix value to open function pointer addresses in our custom 
+	// LLVM pass. At runtime, if we detect the Magic Prefix in the the top 16 bits
+	// of an open function pointer address, we assume that we can find a closed
+	// variant pointer residing 8 bytes before the function address.
+	inline constexpr uint64_t MagicPrefix = 0xa273'0000'0000'0000;
 
 } // namespace AutoRTFM::Constants
 #endif
