[AutoRTFM] Fix stdlib closed mappings for some string functions

A recent update of our toolchain has caused some AutoRTFM open -> closed mappings to fail (specifically linux + debug).
It seems that LibCxx is now making use of __builtin_XXX functions inside header inlined C++ functions. These inlined functions will be split into the usual open + closed pair, and the call to the __builtin_XXX function is now failing. Unfortunately we can't take the address of a __builtin_XXX function, but we do know what these specific functions map to. Do a bit of platform hackery to correctly set up the mappings for these.

Re-enable the PreSubmitTests that were disabled for this issue.

diff --git a/Engine/Source/Runtime/AutoRTFM/Private/StdLibSupport.cpp b/Engine/Source/Runtime/AutoRTFM/Private/StdLibSupport.cpp
--- a/Engine/Source/Runtime/AutoRTFM/Private/StdLibSupport.cpp
+++ b/Engine/Source/Runtime/AutoRTFM/Private/StdLibSupport.cpp
@@ -599,6 +599,17 @@ int RTFM_atexit(void(*Callback)(void)) throw()
 ////////////////////////////////////////////////////////////////////////////////
 // Register all the open -> closed functions
 ////////////////////////////////////////////////////////////////////////////////
+#if AUTORTFM_PLATFORM_LINUX
+// LibCxx's string.h uses builtin string functions inside inline functions.
+// We can't take the address of these builtins, but we know they map to these C functions.
+extern "C"
+{
+	char* strchr(const char*, int);
+	char* strrchr(const char*, int);
+	char* strstr(const char*, const char*);
+}
+#endif
+
 UE_AUTORTFM_REGISTER_OPEN_TO_CLOSED_FUNCTIONS(
 #if AUTORTFM_ASAN_ENABLED
 	UE_AUTORTFM_MAP_OPEN_TO_SELF(__asan_addr_is_in_fake_stack),
@@ -676,6 +687,18 @@ UE_AUTORTFM_REGISTER_OPEN_TO_CLOSED_FUNCTIONS(
 	UE_AUTORTFM_MAP_OPEN_TO_CLOSED(__fxstat, RTFM___fxstat),
 #endif
 
+// Linux requires mapping of the C functions, which do not use const on their return pointers.
+// See comment about LibCxx builtin functions above.
+#if AUTORTFM_PLATFORM_LINUX
+	UE_AUTORTFM_MAP_OPEN_TO_SELF_OVERLOADED(char* (const char*, int), strchr),
+	UE_AUTORTFM_MAP_OPEN_TO_SELF_OVERLOADED(char* (const char*, int), strrchr),
+	UE_AUTORTFM_MAP_OPEN_TO_SELF_OVERLOADED(char* (const char*, const char*), strstr),
+#else
+	UE_AUTORTFM_MAP_OPEN_TO_SELF_OVERLOADED(const char* (const char*, int), strchr),
+	UE_AUTORTFM_MAP_OPEN_TO_SELF_OVERLOADED(const char* (const char*, int), strrchr),
+	UE_AUTORTFM_MAP_OPEN_TO_SELF_OVERLOADED(const char* (const char*, const char*), strstr),
+#endif
+
 	UE_AUTORTFM_MAP_OPEN_TO_CLOSED(memcpy, RTFM_memcpy),
 	UE_AUTORTFM_MAP_OPEN_TO_CLOSED(memmove, RTFM_memmove),
 	UE_AUTORTFM_MAP_OPEN_TO_CLOSED(memset, RTFM_memset),
@@ -690,9 +713,6 @@ UE_AUTORTFM_REGISTER_OPEN_TO_CLOSED_FUNCTIONS(
 	UE_AUTORTFM_MAP_OPEN_TO_SELF(memcmp),
 	UE_AUTORTFM_MAP_OPEN_TO_SELF(strcmp),
 	UE_AUTORTFM_MAP_OPEN_TO_SELF(strncmp),
-	UE_AUTORTFM_MAP_OPEN_TO_SELF_OVERLOADED(const char* (const char*, int), strchr),
-	UE_AUTORTFM_MAP_OPEN_TO_SELF_OVERLOADED(const char* (const char*, int), strrchr),
-	UE_AUTORTFM_MAP_OPEN_TO_SELF_OVERLOADED(const char* (const char*, const char*), strstr),
 	UE_AUTORTFM_MAP_OPEN_TO_SELF(strlen),
 	UE_AUTORTFM_MAP_OPEN_TO_SELF(atof),
 	UE_AUTORTFM_MAP_OPEN_TO_SELF(atoi),
