[AutoRTFM] Include another stack frame in AutoRTFM errors call stacks

Past efforts to clean up the callstack now make it hard to identify AutoRTFM issues from the callstack alone.

diff --git a/Engine/Source/Runtime/AutoRTFM/Private/CompilerInterface.cpp b/Engine/Source/Runtime/AutoRTFM/Private/CompilerInterface.cpp
--- a/Engine/Source/Runtime/AutoRTFM/Private/CompilerInterface.cpp
+++ b/Engine/Source/Runtime/AutoRTFM/Private/CompilerInterface.cpp
@@ -139,18 +139,18 @@ extern "C" UE_AUTORTFM_API void autortfm_memset(void* Dst, int Value, size_t Siz
 
 extern "C" UE_AUTORTFM_API void autortfm_unreachable(const char* Message)
 {
-	ReportError(__FILE__, __LINE__, __builtin_return_address(0), "AutoRTFM Unreachable: %s", Message);
+	AUTORTFM_REPORT_ERROR("AutoRTFM Unreachable: %s", Message);
 }
 
 extern "C" UE_AUTORTFM_API void autortfm_llvm_fail(const char* Message)
 {
 	if (Message)
 	{
-		ReportError(__FILE__, __LINE__, __builtin_return_address(0), "AutoRTFM LLVM Failure: %s", Message);
+		AUTORTFM_REPORT_ERROR("AutoRTFM LLVM Failure: %s", Message);
 	}
 	else
 	{
-		ReportError(__FILE__, __LINE__, __builtin_return_address(0), "AutoRTFM LLVM Failure");
+		AUTORTFM_REPORT_ERROR("AutoRTFM LLVM Failure");
 	}
 }
 
diff --git a/Engine/Source/Runtime/AutoRTFM/Private/FunctionMap.cpp b/Engine/Source/Runtime/AutoRTFM/Private/FunctionMap.cpp
--- a/Engine/Source/Runtime/AutoRTFM/Private/FunctionMap.cpp
+++ b/Engine/Source/Runtime/AutoRTFM/Private/FunctionMap.cpp
@@ -145,11 +145,11 @@ void* FunctionMapReportError(void* OldFunction, const char* Where)
     std::string FunctionDescription = GetFunctionDescription(OldFunction);
 	if (Where)
 	{
-		ReportError(__FILE__, __LINE__, __builtin_return_address(0), "Could not find function %p '%s' where '%s'.", OldFunction, FunctionDescription.c_str(), Where);
+		AUTORTFM_REPORT_ERROR("Could not find function %p '%s' where '%s'.", OldFunction, FunctionDescription.c_str(), Where);
 	}
 	else
 	{
-		ReportError(__FILE__, __LINE__, __builtin_return_address(0), "Could not find function %p '%s'.", OldFunction, FunctionDescription.c_str());
+		AUTORTFM_REPORT_ERROR("Could not find function %p '%s'.", OldFunction, FunctionDescription.c_str());
 	}
 	return nullptr;
 }
