[AutoRTFM] Fix memory leak in AutoRTFM function map.

Whenever we rehashed our function mapping table, we would previously leak the old table instead of freeing it. As the function map would double in size whenever it was rehashed, this leak would double the memory overhead of the function map in practice. (That is, if our active function map contains 32K elements, we would have leaked a 16K-element map, an 8K-element map, a 4K, a 2K, and a 1K along the way.)

This CL also moves rehashing out to a separate function for readability, and replaces `malloc + memset` with `calloc` to improve performance.

diff --git a/Engine/Source/Runtime/AutoRTFM/Private/FunctionMap.cpp b/Engine/Source/Runtime/AutoRTFM/Private/FunctionMap.cpp
--- a/Engine/Source/Runtime/AutoRTFM/Private/FunctionMap.cpp
+++ b/Engine/Source/Runtime/AutoRTFM/Private/FunctionMap.cpp
@@ -49,50 +49,54 @@ FFunctionMap* AllocateFunctionMap(size_t Size)
 {
     AUTORTFM_ASSERT(Size && !(Size & (Size - 1)));
     size_t TotalSize = offsetof(FFunctionMap, Entries) + sizeof(FFunctionEntry) * Size;
-    FFunctionMap* Result = static_cast<FFunctionMap*>(malloc(TotalSize));
+    FFunctionMap* Result = static_cast<FFunctionMap*>(calloc(TotalSize, 1));
     Result->KeyCount = 0;
     Result->Size = Size;
     Result->IndexMask = Size - 1;
-    memset(Result->Entries, 0, sizeof(FFunctionEntry) * Size);
     Fence(); // Really we just need a store-store fence...
     return Result;
 }
 
+static FFunctionMap* RehashFunctionMap(FFunctionMap* Map)
+{
+    FFunctionMap* NewMap = AllocateFunctionMap(Map->Size * 2);
+	for (FFunctionEntry* OldEntry = Map->Entries, *LastOldEntry = Map->Entries + Map->Size; OldEntry != LastOldEntry; ++OldEntry)
+    {
+        if (!OldEntry->OldFunction)
+        {
+            continue;
+        }
+
+        // TODO: If hash collisions are slow, we could have a secondary stronger hash here that
+		// is used on the second iteration followed by +1 for the subsequent iterations.
+        for (size_t Hash = FunctionPtrHash(OldEntry->OldFunction);; ++Hash)
+        {
+            size_t NewIndex = Hash & NewMap->IndexMask;
+            FFunctionEntry* NewEntry = NewMap->Entries + NewIndex;
+            if (!NewEntry->OldFunction)
+            {
+                *NewEntry = *OldEntry;
+                break;
+            }
+        }
+    }
+    NewMap->KeyCount = Map->KeyCount;
+	free(Map);
+    return NewMap;
+}
+
 void FunctionMapAddImpl(void* OldFunction, void* NewFunction)
 {
     FFunctionMap* Map = GlobalData->FunctionMap;
-    if (!Map)
+    if (AUTORTFM_UNLIKELY(!Map))
     {
         Map = AllocateFunctionMap(1024);
         GlobalData->FunctionMap = Map;
     }
-    else if (Map->KeyCount * 2 >= Map->Size)
+    else if (AUTORTFM_UNLIKELY(Map->KeyCount * 2 >= Map->Size))
     {
-        size_t NewSize = Map->Size * 2;
-        FFunctionMap* NewMap = AllocateFunctionMap(NewSize);
-        for (size_t OldIndex = 0; OldIndex < Map->Size; OldIndex++)
-        {
-            FFunctionEntry* OldEntry = Map->Entries + OldIndex;
-            if (!OldEntry->OldFunction)
-            {
-                continue;
-            }
-            // FIXME: Could have a secondary stronger hash here that is used on the second iteration
-            // followed by +1 for the subsequent iterations.
-            for (size_t Hash = FunctionPtrHash(OldEntry->OldFunction); ; ++Hash)
-            {
-                size_t NewIndex = Hash & NewMap->IndexMask;
-                FFunctionEntry* NewEntry = NewMap->Entries + NewIndex;
-                if (!NewEntry->OldFunction)
-                {
-                    *NewEntry = *OldEntry;
-                    break;
-                }
-            }
-        }
-        NewMap->KeyCount = Map->KeyCount;
-        GlobalData->FunctionMap = NewMap;
-        Map = NewMap;
+        Map = RehashFunctionMap(Map);
+        GlobalData->FunctionMap = Map;
     }
 
     for (size_t Hash = FunctionPtrHash(OldFunction); ; ++Hash) {
