[AutoRTFM] Enable (almost) all warnings as errors, fix issues found

diff --git a/Engine/Source/Runtime/AutoRTFM/AutoRTFM.Build.cs b/Engine/Source/Runtime/AutoRTFM/AutoRTFM.Build.cs
--- a/Engine/Source/Runtime/AutoRTFM/AutoRTFM.Build.cs
+++ b/Engine/Source/Runtime/AutoRTFM/AutoRTFM.Build.cs
@@ -15,6 +15,78 @@ namespace UnrealBuildTool.Rules
 				}
 			);
 			PrivateDefinitions.Add("SUPPRESS_PER_MODULE_INLINE_FILE"); // This module does not use core's standard operator new/delete overloads
+
+			// Enable (almost) all warnings as errors
+			CppCompileWarningSettings.ShadowVariableWarningLevel = WarningLevel.Error;
+			CppCompileWarningSettings.UnsafeTypeCastWarningLevel = WarningLevel.Error;
+			CppCompileWarningSettings.UndefinedIdentifierWarningLevel = WarningLevel.Error;
+			// CppCompileWarningSettings.SwitchUnhandledEnumeratorWarningLevel = WarningLevel.Error;
+			CppCompileWarningSettings.EnumEnumConversionWarningLevel = WarningLevel.Error;
+			CppCompileWarningSettings.EnumFloatConversionWarningLevel = WarningLevel.Error;
+			CppCompileWarningSettings.AmbiguousReversedOperatorWarningLevel = WarningLevel.Error;
+			CppCompileWarningSettings.DeprecatedAnonEnumEnumConversionWarningLevel = WarningLevel.Error;
+			CppCompileWarningSettings.DeprecatedVolatileWarningLevel = WarningLevel.Error;
+			CppCompileWarningSettings.OrderedCompareFunctionPointers = WarningLevel.Error;
+			CppCompileWarningSettings.BitwiseInsteadOfLogicalWarningLevel = WarningLevel.Error;
+			CppCompileWarningSettings.DeprecatedCopyWarningLevel = WarningLevel.Error;
+			CppCompileWarningSettings.DeprecatedCopyWithUserProvidedCopyWarningLevel = WarningLevel.Error;
+			CppCompileWarningSettings.InvalidUnevaluatedStringWarningLevel = WarningLevel.Error;
+			CppCompileWarningSettings.NaNInfinityDisabledWarningLevel = WarningLevel.Error;
+			CppCompileWarningSettings.LevelExtraQualificationWarningLevel = WarningLevel.Error;
+			CppCompileWarningSettings.CastFunctionTypeMistmatchWarningLevel = WarningLevel.Error;
+			CppCompileWarningSettings.MissingTemplateArgListAfterTemplateWarningLevel = WarningLevel.Error;
+			CppCompileWarningSettings.GNUStringLiteralOperatorTemplateWarningLevel = WarningLevel.Error;
+			CppCompileWarningSettings.InconsistentMissingOverrideWarningLevel = WarningLevel.Error;
+			CppCompileWarningSettings.InvalidOffsetWarningLevel = WarningLevel.Error;
+			CppCompileWarningSettings.SwitchWarningLevel = WarningLevel.Error;
+			CppCompileWarningSettings.TautologicalCompareWarningLevel = WarningLevel.Error;
+			CppCompileWarningSettings.UnknownPragmasWarningLevel = WarningLevel.Error;
+			CppCompileWarningSettings.UnusedWarningLevel = WarningLevel.Error;
+			CppCompileWarningSettings.UndefinedVarTemplateWarningLevel = WarningLevel.Error;
+			CppCompileWarningSettings.ProfileInstructWarningLevel = WarningLevel.Error;
+			CppCompileWarningSettings.BackendPluginWarningLevel = WarningLevel.Error;
+			CppCompileWarningSettings.UnusedValueWarningLevel = WarningLevel.Error;
+			CppCompileWarningSettings.Shorten64To32WarningLevel = WarningLevel.Error;
+			CppCompileWarningSettings.DllExportExplicitInstantiationDeclWarningLevel = WarningLevel.Error;
+			CppCompileWarningSettings.MicrosoftGroupWarningLevel = WarningLevel.Error;
+			CppCompileWarningSettings.MSVCIncludeWarningLevel = WarningLevel.Error;
+			CppCompileWarningSettings.PragmaPackWarningLevel = WarningLevel.Error;
+			CppCompileWarningSettings.InlineNewDeleteWarningLevel = WarningLevel.Error;
+			CppCompileWarningSettings.ImplicitExceptionSpecMismatchWarningLevel = WarningLevel.Error;
+			CppCompileWarningSettings.UndefinedBoolConversionWarningLevel = WarningLevel.Error;
+			CppCompileWarningSettings.DeprecatedWriteableStringsWarningLevel = WarningLevel.Error;
+			CppCompileWarningSettings.DeprecatedRegisterWarningLevel = WarningLevel.Error;
+			CppCompileWarningSettings.LogicalOpParenthesesWarningLevel = WarningLevel.Error;
+			CppCompileWarningSettings.NullArithmeticWarningLevel = WarningLevel.Error;
+			CppCompileWarningSettings.ReturnTypeCLinkageWarningLevel = WarningLevel.Error;
+			CppCompileWarningSettings.IgnoredAttributesWarningLevel = WarningLevel.Error;
+			CppCompileWarningSettings.UninitializedWarningLevel = WarningLevel.Error;
+			CppCompileWarningSettings.ReturnTypeWarningLevel = WarningLevel.Error;
+			CppCompileWarningSettings.UnusedParameterWarningLevel = WarningLevel.Error;
+			CppCompileWarningSettings.IgnoredQualifiersWarningLevel = WarningLevel.Error;
+			CppCompileWarningSettings.ExpansionToDefined = WarningLevel.Error;
+			CppCompileWarningSettings.SignCompareWarningLevel = WarningLevel.Error;
+			CppCompileWarningSettings.MissingFieldInitializersWarningLevel = WarningLevel.Error;
+			CppCompileWarningSettings.NonPortableIncludePathWarningLevel = WarningLevel.Error;
+			CppCompileWarningSettings.InvalidTokenPasteWarningLevel = WarningLevel.Error;
+			CppCompileWarningSettings.NullPointerArithmeticWarningLevel = WarningLevel.Error;
+			CppCompileWarningSettings.ConstantLogicalOperandWarningLevel = WarningLevel.Error;
+			CppCompileWarningSettings.BitfieldEnumConversion = WarningLevel.Error;
+			CppCompileWarningSettings.NullPointerSubtractionWarningLevel = WarningLevel.Error;
+			CppCompileWarningSettings.DanglingWarningLevel = WarningLevel.Error;
+			CppCompileWarningSettings.MSVCSwitchEnumWarningLevel = WarningLevel.Error;
+			CppCompileWarningSettings.MSVCUnusedValueWarningLevel = WarningLevel.Error;
+			CppCompileWarningSettings.MSVCDeprecationWarningLevel = WarningLevel.Error;
+			CppCompileWarningSettings.DeprecationWarningLevel = WarningLevel.Error;
+			CppCompileWarningSettings.DeterministicWarningLevel = WarningLevel.Error;
+			CppCompileWarningSettings.PCHPerformanceIssueWarningLevel = WarningLevel.Error;
+			CppCompileWarningSettings.ModuleUnsupportedWarningLevel = WarningLevel.Error;
+			CppCompileWarningSettings.PluginModuleUnsupportedWarningLevel = WarningLevel.Error;
+			CppCompileWarningSettings.ModuleIncludePathWarningLevel = WarningLevel.Error;
+			CppCompileWarningSettings.ModuleIncludePrivateWarningLevel = WarningLevel.Error;
+			CppCompileWarningSettings.ModuleIncludeSubdirectoryWarningLevel = WarningLevel.Error;
+			CppCompileWarningSettings.DisablePluginsConflictWarningLevel = WarningLevel.Error;
+			CppCompileWarningSettings.NonInlinedGenCppWarningLevel = WarningLevel.Error;
 		}
 	}
 }
diff --git a/Engine/Source/Runtime/AutoRTFM/Private/CompilerInterface.cpp b/Engine/Source/Runtime/AutoRTFM/Private/CompilerInterface.cpp
--- a/Engine/Source/Runtime/AutoRTFM/Private/CompilerInterface.cpp
+++ b/Engine/Source/Runtime/AutoRTFM/Private/CompilerInterface.cpp
@@ -92,7 +92,6 @@ extern "C" UE_AUTORTFM_API AUTORTFM_NO_ASAN void autortfm_record_write_8(void* P
 extern "C" UE_AUTORTFM_API AUTORTFM_NO_ASAN void autortfm_record_masked_write(void* Ptr, uintptr_t Mask, int MaskWidthBits, int ValueSizeBytes)
 {
 	UE_AUTORTFM_HANDLE_NULL_WRITE(Ptr);
-	FContext* Context = FContext::Get();
 
 	char* IncrementablePtr = static_cast<char*>(Ptr);
 	for(int i = 0; i < MaskWidthBits; i++)
diff --git a/Engine/Source/Runtime/AutoRTFM/Private/Context.cpp b/Engine/Source/Runtime/AutoRTFM/Private/Context.cpp
--- a/Engine/Source/Runtime/AutoRTFM/Private/Context.cpp
+++ b/Engine/Source/Runtime/AutoRTFM/Private/Context.cpp
@@ -97,7 +97,7 @@ void FContext::StartNonDeferredTransaction(EMemoryValidationLevel MemoryValidati
 	AUTORTFM_ASSERT(GetNumDeferredTransactions() == 0);
 	AUTORTFM_ENSURE_MSG(CurrentTransaction, "FContext::StartNonDeferredTransaction() can only be called within a scoped transaction");
 
-	FTransaction* NewTransaction = PushTransaction(
+	PushTransaction(
 		/* Closed */ false,
 		/* bIsScoped */ false,
 		/* StackRange */ CurrentTransaction->GetStackRange(),
@@ -111,8 +111,6 @@ void FContext::StartNonDeferredTransaction(EMemoryValidationLevel MemoryValidati
 
 ETransactionResult FContext::CommitTransaction()
 {
-	constexpr bool bVerbose = false;
-
 	AUTORTFM_ASSERT(Status == EContextStatus::OnTrack);
 
 	ETransactionResult Result = ETransactionResult::Committed;
@@ -414,8 +412,6 @@ AutoRTFM::FStackRange FContext::GetThreadStackRange()
 
 ETransactionResult FContext::Transact(void (*UninstrumentedFunction)(void*), void (*InstrumentedFunction)(void*), void* Arg)
 {
-    constexpr bool bVerbose = false;
-
 	if constexpr (0 != UE_AUTOSTM)
 	{
 		AUTORTFM_WARN("AutoSTM is not implemented in the runtime yet!");
diff --git a/Engine/Source/Runtime/AutoRTFM/Private/ExternAPI.cpp b/Engine/Source/Runtime/AutoRTFM/Private/ExternAPI.cpp
--- a/Engine/Source/Runtime/AutoRTFM/Private/ExternAPI.cpp
+++ b/Engine/Source/Runtime/AutoRTFM/Private/ExternAPI.cpp
@@ -28,7 +28,7 @@ void DefaultLog(const char* File, int Line, [[maybe_unused]] void* ProgramCounte
 	va_list Args2;
 	va_copy(Args2, Args);
 	int Count = vsnprintf(InlineBuffer, InlineBufferLength, Format, Args);
-	auto AllocatedBuffer = (Count >= InlineBufferLength) ? std::unique_ptr<char[]>{new char[Count+1]} : nullptr;
+	auto AllocatedBuffer = (static_cast<size_t>(Count) >= InlineBufferLength) ? std::unique_ptr<char[]>{new char[Count+1]} : nullptr;
 	if (AllocatedBuffer)
 	{
 		vsnprintf(AllocatedBuffer.get(), Count+1, Format, Args2);
@@ -67,6 +67,7 @@ autortfm_extern_api GExternAPI
 	/* AllocateZeroed */ nullptr,
 	/* Free */ nullptr,
 	/* Log */ DefaultLog,
+	/* LogWithCallback */ nullptr,
 	/* EnsureFailure */ nullptr,
 	/* IsLogActive */ nullptr,
 	/* OnRuntimeEnabledChanged */ nullptr,
diff --git a/Engine/Source/Runtime/AutoRTFM/Private/GlobalData.cpp b/Engine/Source/Runtime/AutoRTFM/Private/GlobalData.cpp
--- a/Engine/Source/Runtime/AutoRTFM/Private/GlobalData.cpp
+++ b/Engine/Source/Runtime/AutoRTFM/Private/GlobalData.cpp
@@ -28,7 +28,7 @@ void InitializeGlobalDataIfNecessary()
 #pragma warning(disable : 4996)
 #endif
 
-    static bool bInitialized = []
+    [[maybe_unused]] static bool bInitialized = []
     {
         // AutoRTFM is intended to be used primarily from monolithic UE binaries, where having
         // some additional DLL/so/dylib would just be annoying. However, we guard against the
diff --git a/Engine/Source/Runtime/AutoRTFM/Private/Transaction.cpp b/Engine/Source/Runtime/AutoRTFM/Private/Transaction.cpp
--- a/Engine/Source/Runtime/AutoRTFM/Private/Transaction.cpp
+++ b/Engine/Source/Runtime/AutoRTFM/Private/Transaction.cpp
@@ -182,8 +182,6 @@ void FTransaction::Undo()
 	AUTORTFM_VERBOSE("Undoing a transaction...");
 	AUTORTFM_ASSERT(IsDone());
 
-	int VerboseCounter = 0;
-	int Num = WriteLog.Num();
 	for(auto Iter = WriteLog.rbegin(); Iter != WriteLog.rend(); ++Iter)
     {
 		FWriteLogEntry Entry = *Iter;
@@ -191,8 +189,6 @@ void FTransaction::Undo()
 		AUTORTFM_ENSURE(!IsOnStack(Entry.LogicalAddress));
 
 		memcpy(Entry.LogicalAddress, Entry.Data, Entry.Size);
-
-		VerboseCounter++;
     }
 
 	AUTORTFM_VERBOSE("Undone a transaction!");
diff --git a/Engine/Source/Runtime/AutoRTFM/Private/Utils.cpp b/Engine/Source/Runtime/AutoRTFM/Private/Utils.cpp
--- a/Engine/Source/Runtime/AutoRTFM/Private/Utils.cpp
+++ b/Engine/Source/Runtime/AutoRTFM/Private/Utils.cpp
@@ -71,7 +71,7 @@ UE_AUTORTFM_API UE_AUTORTFM_FORCENOINLINE void ReportError(const char* File, int
 		{
 			va_list Args;
 			va_start(Args, Format);
-			static bool bCalled = [&]
+			[[maybe_unused]] static bool bCalled = [&]
 			{
 				::AutoRTFM::EnsureFailureV(File, Line, ProgramCounter, "!GetEnsureOnInternalAbort()", Format, Args);
 				return true;
@@ -87,7 +87,7 @@ UE_AUTORTFM_API UE_AUTORTFM_FORCENOINLINE void ReportError(const char* File, int
 		}
 		else if (ForTheRuntime::GetEnsureOnInternalAbort())
 		{
-			static bool bCalled = [&]
+			[[maybe_unused]] static bool bCalled = [&]
 			{
 				::AutoRTFM::EnsureFailure(File, Line, ProgramCounter, "!GetEnsureOnInternalAbort()", "Transaction failing because of internal issue");
 				return true;
diff --git a/Engine/Source/Runtime/AutoRTFM/Private/Utils.h b/Engine/Source/Runtime/AutoRTFM/Private/Utils.h
--- a/Engine/Source/Runtime/AutoRTFM/Private/Utils.h
+++ b/Engine/Source/Runtime/AutoRTFM/Private/Utils.h
@@ -211,7 +211,7 @@ inline constexpr T Lerp(T Start, T End, T Fraction)
 
 #define AUTORTFM_ENSURE(Condition) do { \
 	if (AUTORTFM_UNLIKELY(!(Condition))) { \
-		static bool bCalled = [] { \
+		[[maybe_unused]] static bool bCalled = [] { \
 			::AutoRTFM::EnsureFailure(__FILE__, __LINE__, /* ProgramCounter */ nullptr, #Condition, nullptr); \
 			return true; \
 		}(); \
@@ -219,7 +219,7 @@ inline constexpr T Lerp(T Start, T End, T Fraction)
 
 #define AUTORTFM_ENSURE_MSG(Condition, Format, ...) do { \
 	if (AUTORTFM_UNLIKELY(!(Condition))) { \
-		static bool bCalled = [&] { \
+		[[maybe_unused]] static bool bCalled = [&] { \
 			::AutoRTFM::EnsureFailure(__FILE__, __LINE__, /* ProgramCounter */ nullptr, #Condition, Format, ##__VA_ARGS__); \
 			return true; \
 		}(); \
@@ -227,7 +227,7 @@ inline constexpr T Lerp(T Start, T End, T Fraction)
 
 #define AUTORTFM_ENSURE_MSG_V(Condition, Format, Args) do { \
 	if (AUTORTFM_UNLIKELY(!(Condition))) { \
-		static bool bCalled = [&] { \
+		[[maybe_unused]] static bool bCalled = [&] { \
 			::AutoRTFM::EnsureFailureV(__FILE__, __LINE__, /* ProgramCounter */ nullptr, #Condition, Format, Args); \
 			return true; \
 		}(); \
diff --git a/Engine/Source/Runtime/AutoRTFM/Private/WriteLog.cpp b/Engine/Source/Runtime/AutoRTFM/Private/WriteLog.cpp
--- a/Engine/Source/Runtime/AutoRTFM/Private/WriteLog.cpp
+++ b/Engine/Source/Runtime/AutoRTFM/Private/WriteLog.cpp
@@ -119,7 +119,6 @@ FWriteLog::FHash FWriteLog::HashAVX2(size_t NumWriteEntries) const
 	FWriteLog::FHash Hash = 0;
 
 	using i8x32 = __m256i;
-	using u64x4 = __m256i;
 	const i8x32 Vec0To31 = _mm256_setr_epi8( // [0..31]
 		0x00, 0x01, 0x02, 0x03, 0x04, 0x05, 0x06, 0x07,
 		0x08, 0x09, 0x0a, 0x0b, 0x0c, 0x0d, 0x0e, 0x0f,
diff --git a/Engine/Source/Runtime/AutoRTFM/Public/AutoRTFM.h b/Engine/Source/Runtime/AutoRTFM/Public/AutoRTFM.h
--- a/Engine/Source/Runtime/AutoRTFM/Public/AutoRTFM.h
+++ b/Engine/Source/Runtime/AutoRTFM/Public/AutoRTFM.h
@@ -1461,7 +1461,7 @@ namespace ForTheRuntime
 				Mappings[Count++] = {nullptr, nullptr};
 			}
 			autortfm_open_to_closed_mapping Mappings[sizeof...(OpenToClosedMappings) + 1]{};
-			autortfm_open_to_closed_table Table{Mappings, nullptr};
+			autortfm_open_to_closed_table Table{Mappings, /* Prev */ nullptr, /* Next */ nullptr};
 		};
 
 		UE_AUTORTFM_NOAUTORTFM
