[AutoRTFM] Fix implicit conversion warning treated as error

diff --git a/Engine/Source/Runtime/AutoRTFM/Private/IntervalTree.h b/Engine/Source/Runtime/AutoRTFM/Private/IntervalTree.h
--- a/Engine/Source/Runtime/AutoRTFM/Private/IntervalTree.h
+++ b/Engine/Source/Runtime/AutoRTFM/Private/IntervalTree.h
@@ -140,7 +140,7 @@ private:
 
 	UE_AUTORTFM_FORCEINLINE FIntervalTreeNodeIndex AddNode(FRange&& NewRange)
 	{
-		const FIntervalTreeNodeIndex Index = NodeRanges.Num();
+		const FIntervalTreeNodeIndex Index = static_cast<FIntervalTreeNodeIndex>(NodeRanges.Num());
         NodeRanges.Push(NewRange);
 		NodeLefts.Push(IntervalTreeNodeIndexNone);
 		NodeRights.Push(IntervalTreeNodeIndexNone);
@@ -151,7 +151,7 @@ private:
 
 	UE_AUTORTFM_FORCEINLINE FIntervalTreeNodeIndex AddNode(FRange&& NewRange, FIntervalTreeNodeIndex Parent)
 	{
-		const FIntervalTreeNodeIndex Index = NodeRanges.Num();
+		const FIntervalTreeNodeIndex Index = static_cast<FIntervalTreeNodeIndex>(NodeRanges.Num());
         NodeRanges.Push(NewRange);
 		NodeLefts.Push(IntervalTreeNodeIndexNone);
 		NodeRights.Push(IntervalTreeNodeIndexNone);
