[AutoRTFM] turn off backend plugin warnings-as-errors to make PGO builds stop complaining about hash mismatches

#rnx

diff --git a/Engine/Source/Runtime/AutoRTFM/AutoRTFM.Build.cs b/Engine/Source/Runtime/AutoRTFM/AutoRTFM.Build.cs
--- a/Engine/Source/Runtime/AutoRTFM/AutoRTFM.Build.cs
+++ b/Engine/Source/Runtime/AutoRTFM/AutoRTFM.Build.cs
@@ -44,7 +44,7 @@ namespace UnrealBuildTool.Rules
 			CppCompileWarningSettings.UnusedWarningLevel = WarningLevel.Error;
 			CppCompileWarningSettings.UndefinedVarTemplateWarningLevel = WarningLevel.Error;
 			CppCompileWarningSettings.ProfileInstructWarningLevel = WarningLevel.Error;
-			CppCompileWarningSettings.BackendPluginWarningLevel = WarningLevel.Error;
+			// CppCompileWarningSettings.BackendPluginWarningLevel = WarningLevel.Error;
 			CppCompileWarningSettings.UnusedValueWarningLevel = WarningLevel.Error;
 			CppCompileWarningSettings.Shorten64To32WarningLevel = WarningLevel.Error;
 			CppCompileWarningSettings.DllExportExplicitInstantiationDeclWarningLevel = WarningLevel.Error;
