[AutoRTFM] Improve hash function used by the function map.

We now use a Fibonacci hash to reduce the chance of collisions. In order to quickly extract the top bits from the hash, we now store "2^Size" inside the FunctionMap instead of the actual size; in other words, a function map containing 2^10 elements would store Size2n = 10, not Size = 1024.

diff --git a/Engine/Source/Programs/AutoRTFMEngineTests/Private/AutoRTFMEngineTests.cpp b/Engine/Source/Programs/AutoRTFMEngineTests/Private/AutoRTFMEngineTests.cpp
--- a/Engine/Source/Programs/AutoRTFMEngineTests/Private/AutoRTFMEngineTests.cpp
+++ b/Engine/Source/Programs/AutoRTFMEngineTests/Private/AutoRTFMEngineTests.cpp
@@ -22,6 +22,11 @@ typedef UE::AutoRTFM::FAutomationTestRunner TestRunner;
 typedef FNullTestRunner TestRunner;
 #endif
 
+namespace AutoRTFM
+{
+	void FunctionMapDumpStats();
+}
+
 DEFINE_LOG_CATEGORY(LogAutoRTFMEngineTests);
 
 IMPLEMENT_APPLICATION(AutoRTFMEngineTests, "AutoRTFMEngineTests");
@@ -114,6 +119,8 @@ INT32_MAIN_INT32_ARGC_TCHAR_ARGV()
 
 	}
 
+	/* AutoRTFM::FunctionMapDumpStats(); */
+
 	TearDown();
 
 	return TestsPassed ? 0 : 1;
diff --git a/Engine/Source/Runtime/AutoRTFM/Private/FunctionMap.cpp b/Engine/Source/Runtime/AutoRTFM/Private/FunctionMap.cpp
--- a/Engine/Source/Runtime/AutoRTFM/Private/FunctionMap.cpp
+++ b/Engine/Source/Runtime/AutoRTFM/Private/FunctionMap.cpp
@@ -11,10 +11,11 @@
 #include "LockGuard.h"
 #include "Utils.h"
 
+#include <iterator>
+#include <map>
 #include <stdlib.h>
 #include <string.h>
 
-
 namespace AutoRTFM
 {
 
@@ -45,22 +46,21 @@ static inline void* FollowRelocation(void* Function)
 	return nullptr;
 }
 
-FFunctionMap* AllocateFunctionMap(size_t Size)
+static FFunctionMap* AllocateFunctionMap(size_t Size2n)
 {
-    AUTORTFM_ASSERT(Size && !(Size & (Size - 1)));
-    size_t TotalSize = offsetof(FFunctionMap, Entries) + sizeof(FFunctionEntry) * Size;
+    size_t TotalSize = offsetof(FFunctionMap, Entries) + (sizeof(FFunctionEntry) << Size2n);
     FFunctionMap* Result = static_cast<FFunctionMap*>(calloc(TotalSize, 1));
     Result->KeyCount = 0;
-    Result->Size = Size;
-    Result->IndexMask = Size - 1;
+    Result->Size2n = Size2n;
+    Result->IndexMask = (1 << Size2n) - 1;
     Fence(); // Really we just need a store-store fence...
     return Result;
 }
 
 static FFunctionMap* RehashFunctionMap(FFunctionMap* Map)
 {
-    FFunctionMap* NewMap = AllocateFunctionMap(Map->Size * 2);
-	for (FFunctionEntry* OldEntry = Map->Entries, *LastOldEntry = Map->Entries + Map->Size; OldEntry != LastOldEntry; ++OldEntry)
+	FFunctionMap* NewMap = AllocateFunctionMap(Map->Size2n + 1);
+	for (FFunctionEntry* OldEntry = Map->Entries, *FinalEntry = Map->Entries + Map->IndexMask; OldEntry <= FinalEntry; ++OldEntry)
     {
         if (!OldEntry->OldFunction)
         {
@@ -69,10 +69,9 @@ static FFunctionMap* RehashFunctionMap(FFunctionMap* Map)
 
         // TODO: If hash collisions are slow, we could have a secondary stronger hash here that
 		// is used on the second iteration followed by +1 for the subsequent iterations.
-        for (size_t Hash = FunctionPtrHash(OldEntry->OldFunction);; ++Hash)
+        for (size_t Index = FunctionPtrHash(OldEntry->OldFunction, NewMap->Size2n);; Index = (Index + 1) & NewMap->IndexMask)
         {
-            size_t NewIndex = Hash & NewMap->IndexMask;
-            FFunctionEntry* NewEntry = NewMap->Entries + NewIndex;
+            FFunctionEntry* NewEntry = NewMap->Entries + Index;
             if (!NewEntry->OldFunction)
             {
                 *NewEntry = *OldEntry;
@@ -85,32 +84,36 @@ static FFunctionMap* RehashFunctionMap(FFunctionMap* Map)
     return NewMap;
 }
 
-void FunctionMapAddImpl(void* OldFunction, void* NewFunction)
+static void FunctionMapAddImpl(void* OldFunction, void* NewFunction)
 {
     FFunctionMap* Map = GlobalData->FunctionMap;
     if (AUTORTFM_UNLIKELY(!Map))
     {
-        Map = AllocateFunctionMap(1024);
+		// We have no function map at all; create an empty map.
+        Map = AllocateFunctionMap(/*Size2n=*/ 10);
         GlobalData->FunctionMap = Map;
     }
-    else if (AUTORTFM_UNLIKELY(Map->KeyCount * 2 >= Map->Size))
+    else if (AUTORTFM_UNLIKELY(Map->KeyCount * 2 > Map->IndexMask))
     {
+		// The function map is above 50% capacity; we need to grow and rehash.
         Map = RehashFunctionMap(Map);
         GlobalData->FunctionMap = Map;
     }
 
-    for (size_t Hash = FunctionPtrHash(OldFunction); ; ++Hash) {
-        size_t Index = Hash & Map->IndexMask;
+    for (size_t Index = FunctionPtrHash(OldFunction, Map->Size2n);; Index = (Index + 1) & Map->IndexMask)
+	{
         FFunctionEntry* Entry = Map->Entries + Index;
 
-        if (!Entry->OldFunction) {
+        if (!Entry->OldFunction) 
+		{
             Entry->OldFunction = OldFunction;
             Entry->NewFunction = NewFunction;
             Map->KeyCount++;
             return;
         }
 
-        if (Entry->OldFunction == OldFunction) {
+        if (Entry->OldFunction == OldFunction) 
+		{
             Entry->NewFunction = NewFunction;
             return;
         }
@@ -135,6 +138,20 @@ void FunctionMapAdd(void* OldFunction, void* NewFunction)
 	}
 }
 
+static void* FunctionMapReportError(void* OldFunction, const char* Where)
+{
+    std::string FunctionDescription = GetFunctionDescription(OldFunction);
+	if (Where)
+	{
+		AUTORTFM_REPORT_ERROR("Could not find function %p '%s' where '%s'.", OldFunction, FunctionDescription.c_str(), Where);
+	}
+	else
+	{
+		AUTORTFM_REPORT_ERROR("Could not find function %p '%s'.", OldFunction, FunctionDescription.c_str());
+	}
+	return nullptr;
+}
+
 void* FunctionMapLookupTryFollowingRelocation(void* Function, const char* Where)
 {
 	if (void* Ptr = FollowRelocation(Function); AUTORTFM_LIKELY(Ptr))
@@ -144,18 +161,37 @@ void* FunctionMapLookupTryFollowingRelocation(void* Function, const char* Where)
 	AUTORTFM_MUST_TAIL return FunctionMapReportError(Function, Where);
 }
 
-void* FunctionMapReportError(void* OldFunction, const char* Where)
+void FunctionMapDumpStats()
 {
-    std::string FunctionDescription = GetFunctionDescription(OldFunction);
-	if (Where)
+    const FFunctionMap* const Map = GlobalData->FunctionMap;
+	size_t MapSlots = 1llu << Map->Size2n;
+	AUTORTFM_LOG("Function Map Stats");
+	AUTORTFM_LOG("==================");
+	AUTORTFM_LOG("Occupancy: %zu entries in %llu slots (load factor %g%%)", 
+		         Map->KeyCount, 
+		         MapSlots,
+				 double(Map->KeyCount * 1000llu / MapSlots) / 10.0);
+
+	std::map<int, int> CollisionMap; // <probes required to find the function, number of functions>
+	for (const FFunctionEntry* Entry = Map->Entries, *FinalEntry = Map->Entries + Map->IndexMask; Entry <= FinalEntry; ++Entry)
 	{
-		AUTORTFM_REPORT_ERROR("Could not find function %p '%s' where '%s'.", OldFunction, FunctionDescription.c_str(), Where);
+		if (Entry->OldFunction)
+		{
+			uintptr_t ActualIndex = std::distance(Map->Entries, Entry);
+			uintptr_t IdealIndex = FunctionPtrHash(Entry->OldFunction, Map->Size2n) & Map->IndexMask;
+			// We need to account for wraparound in this calculation.
+			uintptr_t Delta = (IdealIndex <= ActualIndex) 
+				? ActualIndex - IdealIndex
+				: ActualIndex + MapSlots - IdealIndex;
+				                                           
+			CollisionMap[Delta] += 1;
+		}
 	}
-	else
+
+	for (int NumProbes = 0, HighestProbes = CollisionMap.rbegin()->first; NumProbes <= HighestProbes; ++NumProbes)
 	{
-		AUTORTFM_REPORT_ERROR("Could not find function %p '%s'.", OldFunction, FunctionDescription.c_str());
+		AUTORTFM_LOG("% 2d probes: %d functions", NumProbes, CollisionMap[NumProbes]); 
 	}
-	return nullptr;
 }
 
 } // namespace AutoRTFM
diff --git a/Engine/Source/Runtime/AutoRTFM/Private/FunctionMap.h b/Engine/Source/Runtime/AutoRTFM/Private/FunctionMap.h
--- a/Engine/Source/Runtime/AutoRTFM/Private/FunctionMap.h
+++ b/Engine/Source/Runtime/AutoRTFM/Private/FunctionMap.h
@@ -5,54 +5,54 @@
 #if (defined(__AUTORTFM) && __AUTORTFM)
 
 #include "GlobalData.h"
+#include "Utils.h"
 
 namespace AutoRTFM
 {
 
-class FContext;
-
-// The compiler has to agree to this.
-inline uintptr_t FunctionPtrHash(void* functionPtr)
+inline uintptr_t FunctionPtrHash(void* FunctionPtr, size_t HashTableSize2n)
 {
-    // FIXME: We'll need something better.
-    return reinterpret_cast<uintptr_t>(functionPtr) / 4;
+	// Use a Fibonacci hash to fold our 64-bit value into a hash-appropriate value.
+	// https://probablydance.com/2018/06/16/fibonacci-hashing-the-optimization-that-the-world-forgot-or-a-better-alternative-to-integer-modulo/
+	uintptr_t HashBits = reinterpret_cast<uintptr_t>(FunctionPtr);
+	// Apply Fibonacci product and preserve the highest bits.
+	return (HashBits * 11400714819323198485llu) >> (64 - HashTableSize2n);
 }
 
 struct FFunctionEntry
 {
-    FFunctionEntry() = default;
-
-    FFunctionEntry(void* OldFunction, void* NewFunction)
-        : OldFunction(OldFunction), NewFunction(NewFunction) { }
-    
-    void* OldFunction{nullptr};
-    void* NewFunction{nullptr};
+    void* OldFunction = nullptr;
+	void* NewFunction = nullptr;
 };
 
-struct FFunctionMap {
-    size_t KeyCount;
-    size_t Size;
-    size_t IndexMask;
-    FFunctionEntry Entries[1];
+struct FFunctionMap 
+{
+	size_t KeyCount;           // The number of elements in the map.
+	size_t Size2n;             // The map size as a power-of-two; that is, a Size2n of 10 corresponds to 1024 entries, as 2^10 == 1024.
+	size_t IndexMask;          // The map size minus one, for use as a bitmask; a function map with 1024 elements will have an IndexMask of 0x3FF.
+    FFunctionEntry Entries[1]; // Actually contains `2 ^ Size2n` elements.
 };
 
+// Dumps statistics about the function map to LOG_INFO.
+void FunctionMapDumpStats();
+
 void FunctionMapAdd(void* OldFunction, void* NewFunction);
 
 // Attempts to lookup the function through a relocation thunk.
 // This is called by FunctionMapLookup(), and should not be called directly.
 void* FunctionMapLookupTryFollowingRelocation(void* Function, const char* Where);
 
-void* FunctionMapReportError(void* OldFunction, const char* Where);
-
 inline void* FunctionMapTryLookup(void* OldFunction)
 {
     FFunctionMap* Map = GlobalData->FunctionMap;
 
-    if (!Map)
+    if (AUTORTFM_UNLIKELY(!Map))
+	{
         return nullptr;
+	}
 
-    for (size_t Hash = FunctionPtrHash(OldFunction); ; ++Hash) {
-        size_t Index = Hash & Map->IndexMask;
+    for (size_t Index = FunctionPtrHash(OldFunction, Map->Size2n);; Index = (Index + 1) & Map->IndexMask) 
+	{
         FFunctionEntry* Entry = Map->Entries + Index;
 
         if (Entry->OldFunction == OldFunction)
