[AutoRTFM] Fix Horde error "implicit conversion loses integer precision".

https://horde.devtools.epicgames.com/job/67f68e3b363e9f7aa0d863fc?step=3f2e&issue=862387

diff --git a/Engine/Source/Runtime/AutoRTFM/Private/FunctionMap.cpp b/Engine/Source/Runtime/AutoRTFM/Private/FunctionMap.cpp
--- a/Engine/Source/Runtime/AutoRTFM/Private/FunctionMap.cpp
+++ b/Engine/Source/Runtime/AutoRTFM/Private/FunctionMap.cpp
@@ -11,6 +11,7 @@
 #include "LockGuard.h"
 #include "Utils.h"
 
+#include <cinttypes>
 #include <iterator>
 #include <map>
 #include <stdlib.h>
@@ -172,7 +173,7 @@ void FunctionMapDumpStats()
 		         MapSlots,
 				 double(Map->KeyCount * 1000llu / MapSlots) / 10.0);
 
-	std::map<int, int> CollisionMap; // <probes required to find the function, number of functions>
+	std::map<uintptr_t, size_t> CollisionMap; // <probes required to find the function, number of functions>
 	for (const FFunctionEntry* Entry = Map->Entries, *FinalEntry = Map->Entries + Map->IndexMask; Entry <= FinalEntry; ++Entry)
 	{
 		if (Entry->OldFunction)
@@ -188,9 +189,9 @@ void FunctionMapDumpStats()
 		}
 	}
 
-	for (int NumProbes = 0, HighestProbes = CollisionMap.rbegin()->first; NumProbes <= HighestProbes; ++NumProbes)
+	for (uintptr_t NumProbes = 0, HighestProbes = CollisionMap.rbegin()->first; NumProbes <= HighestProbes; ++NumProbes)
 	{
-		AUTORTFM_LOG("% 2d probes: %d functions", NumProbes, CollisionMap[NumProbes]); 
+		AUTORTFM_LOG("%2" PRIuPTR " probes: %zu functions", NumProbes, CollisionMap[NumProbes]); 
 	}
 }
 
