[AutoRTFM] Fix Horde error with format specifier mixup.

https://horde.devtools.epicgames.com/log/67f6cbfd335ba6e29e8dbe35?lineIndex=215

diff --git a/Engine/Source/Runtime/AutoRTFM/Private/FunctionMap.cpp b/Engine/Source/Runtime/AutoRTFM/Private/FunctionMap.cpp
--- a/Engine/Source/Runtime/AutoRTFM/Private/FunctionMap.cpp
+++ b/Engine/Source/Runtime/AutoRTFM/Private/FunctionMap.cpp
@@ -168,7 +168,7 @@ void FunctionMapDumpStats()
 	size_t MapSlots = 1llu << Map->Size2n;
 	AUTORTFM_LOG("Function Map Stats");
 	AUTORTFM_LOG("==================");
-	AUTORTFM_LOG("Occupancy: %zu entries in %llu slots (load factor %g%%)", 
+	AUTORTFM_LOG("Occupancy: %zu entries in %zu slots (load factor %g%%)", 
 		         Map->KeyCount, 
 		         MapSlots,
 				 double(Map->KeyCount * 1000llu / MapSlots) / 10.0);
