[AutoRTFM] Remove GlobalData, FastLock & Atomic

This was all used to hold the function map, but this code is unnecessary and complex:
* Using an env var to pass pointers between DLLs / shared objects is bonkers and is not needed. AUTORTFM_CORE_API (dllexport / import) takes care of dynamic linking.
* The FastLock was being used to lock the function map when it was being added to, but not when it was being read from, so this wasn't really guarding concurrent use at all.
* A "fast" spin lock is not really needed. We only hit the function map if the magic-prefix method fails, and that is only done once per process.

This cleanup also reduces the number of times we need to lock the mutex from once per registered function, to once for the entire process in a monolithic build, and once per CU for DLLs. With some LLVM tricks we can probably reduce this down to once per DLL, but this is probably unnecessary.

Other changes include reducing the amount of code that gets inlined for function lookups. While probably mostly noise, benchmarks suggest ~5% performance improvement with this change.

diff --git a/Engine/Source/Runtime/AutoRTFM/Private/API.cpp b/Engine/Source/Runtime/AutoRTFM/Private/API.cpp
--- a/Engine/Source/Runtime/AutoRTFM/Private/API.cpp
+++ b/Engine/Source/Runtime/AutoRTFM/Private/API.cpp
@@ -673,16 +673,8 @@ AUTORTFM_CORE_API bool IsAutoRTFMInitialized()
 AUTORTFM_CORE_API void ProcessAllPendingOpenToClosedRegistrations()
 {
 	AUTORTFM_ASSERT(IsAutoRTFMInitialized());
-
-	while (autortfm_open_to_closed_table* Table = GPendingOpenToClosedRegistrations)
-	{
-		for (const autortfm_open_to_closed_mapping* Mapping = Table->Mappings; Mapping->Open; Mapping++)
-		{
-			AUTORTFM_VERBOSE("Registering open %p -> %p", Mapping->Open, Mapping->Closed);
-			FunctionMapAdd(Mapping->Open, Mapping->Closed);
-		}
-		GPendingOpenToClosedRegistrations = Table->Next;
-	}
+	FunctionMapAdd(GPendingOpenToClosedRegistrations);
+	GPendingOpenToClosedRegistrations = nullptr;
 }
 
 }  // anonymous namespace
diff --git a/Engine/Source/Runtime/AutoRTFM/Private/Atomic.h b/Engine/Source/Runtime/AutoRTFM/Private/Atomic.h
deleted file mode 100644
diff --git a/Engine/Source/Runtime/AutoRTFM/Private/FastLock.h b/Engine/Source/Runtime/AutoRTFM/Private/FastLock.h
deleted file mode 100644
diff --git a/Engine/Source/Runtime/AutoRTFM/Private/FunctionMap.cpp b/Engine/Source/Runtime/AutoRTFM/Private/FunctionMap.cpp
--- a/Engine/Source/Runtime/AutoRTFM/Private/FunctionMap.cpp
+++ b/Engine/Source/Runtime/AutoRTFM/Private/FunctionMap.cpp
@@ -4,26 +4,230 @@
 
 #include "FunctionMap.h"
 
-#include "Atomic.h"
+#include "AutoRTFM.h"
 #include "BuildMacros.h"
-#include "FastLock.h"
 #include "FunctionMapInlines.h"
-#include "LockGuard.h"
 #include "Utils.h"
 
 #include <cinttypes>
+#include <cstddef>
 #include <iterator>
 #include <map>
+#include <mutex>
 #include <stdlib.h>
 #include <string.h>
 
 namespace AutoRTFM
 {
 
+namespace
+{
+
+// The function map holds a hash map of open function pointer to closed function
+// pointer. The function map must be obtained by calling FFunctionMap::Get(),
+// which holds an internal mutex, preventing concurrent access.
+class FFunctionMap
+{
+	// The internal hash map load factor.
+	// Currently capacity must be at least twice the number of entries.
+	static constexpr size_t LoadFactor = 2;
+
+	// The number of bits in a size_t integer.
+	static constexpr size_t NumBitsInSizeT = sizeof(size_t) * 8;
+
+	// The internal hash map data.
+	struct FHashMap
+	{
+		// The number of elements in the map.
+		size_t EntryCount;
+		// The map capacity as a power-of-two; that is, a Capacity2n of 10 corresponds to 1024 entries, as 2^10 == 1024.
+		size_t Capacity2n;
+		// The map capacity minus one, for use as a bitmask; a function map with 1024 elements will have an IndexMask of 0x3FF.
+		size_t IndexMask;
+		// Actually contains `2 ^ Capacity2n` elements.
+		autortfm_open_to_closed_mapping Entries[1];
+
+		static FHashMap* Allocate(size_t Capacity2n)
+		{
+			size_t TotalSize = offsetof(FHashMap, Entries) + (sizeof(autortfm_open_to_closed_mapping) << Capacity2n);
+			FHashMap* Result = static_cast<FHashMap*>(calloc(TotalSize, 1));
+			Result->Capacity2n = Capacity2n;
+			Result->IndexMask = (1 << Capacity2n) - 1;
+			return Result;
+		}
+	};
+
+	// Returns the number of bits required to represent the given value (must be non-zero)
+	static inline size_t NumberOfBits(size_t Value)
+	{
+		return NumBitsInSizeT - __builtin_clzll(Value);
+	}
+	
+public:
+	// Obtains the FFunctionMap instance.
+	static FFunctionMap Get()
+	{
+		struct FHashMapAndMutex
+		{
+			FHashMap* Data = FHashMap::Allocate(/* Capacity2n */ 10);
+			std::mutex Mutex{};
+		};
+		static FHashMapAndMutex HashMapAndMutex;
+		return FFunctionMap{HashMapAndMutex.Data, HashMapAndMutex.Mutex};
+	}
+
+	~FFunctionMap()
+	{
+		Mutex.unlock();
+	}
+
+	// Returns the total number of entries in the function map.
+	inline size_t Count() const
+	{
+		return Map->EntryCount;
+	}
+
+	// Ensures that the function map is large enough to hold Count * LoadFactor entries.
+	void Reserve(size_t Count)
+	{
+		const size_t NewCapacity = Count * LoadFactor;
+		const size_t NewCapacity2n = NumberOfBits(NewCapacity);
+		if (Map->Capacity2n >= NewCapacity2n)
+		{
+			return;
+		}
+
+		FHashMap* NewMap = FHashMap::Allocate(NewCapacity2n);
+		for (autortfm_open_to_closed_mapping* OldEntry = Map->Entries, *FinalEntry = Map->Entries + Map->IndexMask; OldEntry <= FinalEntry; ++OldEntry)
+		{
+			if (!OldEntry->Open)
+			{
+				continue;
+			}
+	
+			// TODO: If hash collisions are slow, we could have a secondary stronger hash here that
+			// is used on the second iteration followed by +1 for the subsequent iterations.
+			for (size_t Index = FunctionPtrHash(OldEntry->Open, NewCapacity2n);; Index = (Index + 1) & NewMap->IndexMask)
+			{
+				autortfm_open_to_closed_mapping* NewEntry = NewMap->Entries + Index;
+				if (!NewEntry->Open)
+				{
+					*NewEntry = *OldEntry;
+					break;
+				}
+			}
+		}
+		NewMap->EntryCount = Map->EntryCount;
+		free(Map);
+		Map = NewMap;
+	}
+
+	// Adds a new entry to the map.
+	// The map must have double the capacity of the number of entries before calling.
+	void Add(void* Open, void* Closed)
+	{
+		AUTORTFM_ASSERT(Map->EntryCount + 1 <= (1 << Map->Capacity2n));
+
+		for (size_t Index = FunctionPtrHash(Open, Map->Capacity2n);; Index = (Index + 1) & Map->IndexMask)
+		{
+			autortfm_open_to_closed_mapping* Mapping = Map->Entries + Index;
+
+			if (!Mapping->Open) 
+			{
+				Mapping->Open = Open;
+				Mapping->Closed = Closed;
+				Map->EntryCount++;
+				return;
+			}
+
+			if (Mapping->Open == Open) 
+			{
+				Mapping->Closed = Closed;
+				return;
+			}
+		}
+	}
+	
+	// Looks up the closed function from the open function pointer.
+	// Returns nullptr if the mapping is not found.
+	void* Lookup(void* OpenFn)
+	{
+		for (size_t Index = FunctionPtrHash(OpenFn, Map->Capacity2n);; Index = (Index + 1) & Map->IndexMask) 
+		{
+			autortfm_open_to_closed_mapping* Entry = Map->Entries + Index;
+
+			if (Entry->Open == OpenFn)
+			{
+				return Entry->Closed;
+			}
+
+			if (!Entry->Open)
+			{
+				return nullptr;
+			}
+		}
+	}
+
+	void DumpStats()
+	{
+		size_t MapSlots = 1llu << Map->Capacity2n;
+		AUTORTFM_LOG("Function Map Stats");
+		AUTORTFM_LOG("==================");
+		AUTORTFM_LOG("Occupancy: %zu entries in %zu slots (load factor %g%%)", 
+					 Map->EntryCount, 
+					 MapSlots,
+					 static_cast<double>(Map->EntryCount * 1000llu / MapSlots) / 10.0);
+	
+		std::map<uintptr_t, size_t> CollisionMap; // <probes required to find the function, number of functions>
+		for (const autortfm_open_to_closed_mapping* Entry = Map->Entries, *FinalEntry = Map->Entries + Map->IndexMask; Entry <= FinalEntry; ++Entry)
+		{
+			if (Entry->Open)
+			{
+				uintptr_t ActualIndex = std::distance<const autortfm_open_to_closed_mapping*>(Map->Entries, Entry);
+				uintptr_t IdealIndex = FunctionPtrHash(Entry->Open, Map->Capacity2n) & Map->IndexMask;
+				// We need to account for wraparound in this calculation.
+				uintptr_t Delta = (IdealIndex <= ActualIndex) 
+					? ActualIndex - IdealIndex
+					: ActualIndex + MapSlots - IdealIndex;
+															   
+				CollisionMap[Delta] += 1;
+			}
+		}
+	
+		for (uintptr_t NumProbes = 0, HighestProbes = CollisionMap.rbegin()->first; NumProbes <= HighestProbes; ++NumProbes)
+		{
+			AUTORTFM_LOG("%2" PRIuPTR " probes: %zu functions", NumProbes, CollisionMap[NumProbes]); 
+		}
+	}
+
+private:
+	FFunctionMap(const FFunctionMap&) = delete;
+	FFunctionMap& operator = (const FFunctionMap&) = delete;
+	FFunctionMap(FFunctionMap&&) = delete;
+	FFunctionMap& operator = (FFunctionMap&&) = delete;
+
+	FFunctionMap(FHashMap*& Map, std::mutex& Mutex) : Map{Map}, Mutex{Mutex}
+	{
+		Mutex.lock();
+	}
+	
+	inline static uintptr_t FunctionPtrHash(void* FunctionPtr, size_t HashTableSize2n)
+	{
+		// Use a Fibonacci hash to fold our size_t value into a hash-appropriate value.
+		// https://probablydance.com/2018/06/16/fibonacci-hashing-the-optimization-that-the-world-forgot-or-a-better-alternative-to-integer-modulo/
+		uintptr_t HashBits = reinterpret_cast<uintptr_t>(FunctionPtr);
+		// Apply Fibonacci product and preserve the highest bits.
+		return (HashBits * 11400714819323198485llu) >> (NumBitsInSizeT - HashTableSize2n);
+	}
+
+	FHashMap*& Map;
+	std::mutex& Mutex;
+};
+
 // Attempts to lookup the "true" function from a dynamically linked import function thunk.
 // Returns the pointer to the true function, or nullptr if the function cannot 
 // be resolved (or is not an import thunk).
-static inline void* FollowRelocation(void* Function)
+inline void* FollowRelocation(void* Function)
 {
 #if AUTORTFM_PLATFORM_WINDOWS
 #if defined(_x86_64) || defined(__x86_64__) || defined(_M_X64) || defined(_M_AMD64)
@@ -47,152 +251,85 @@ static inline void* FollowRelocation(void* Function)
 	return nullptr;
 }
 
-static FFunctionMap* AllocateFunctionMap(size_t Size2n)
+void* FunctionMapReportError(void* Open, const char* Where)
 {
-    size_t TotalSize = offsetof(FFunctionMap, Entries) + (sizeof(FFunctionEntry) << Size2n);
-    FFunctionMap* Result = static_cast<FFunctionMap*>(calloc(TotalSize, 1));
-    Result->KeyCount = 0;
-    Result->Size2n = Size2n;
-    Result->IndexMask = (1 << Size2n) - 1;
-    Fence(); // Really we just need a store-store fence...
-    return Result;
+	std::string FunctionDescription = GetFunctionDescription(Open);
+	if (Where)
+	{
+		AUTORTFM_REPORT_ERROR("Could not find function %p '%s' where '%s'.", Open, FunctionDescription.c_str(), Where);
+	}
+	else
+	{
+		AUTORTFM_REPORT_ERROR("Could not find function %p '%s'.", Open, FunctionDescription.c_str());
+	}
+	return nullptr;
 }
 
-static FFunctionMap* RehashFunctionMap(FFunctionMap* Map)
-{
-	FFunctionMap* NewMap = AllocateFunctionMap(Map->Size2n + 1);
-	for (FFunctionEntry* OldEntry = Map->Entries, *FinalEntry = Map->Entries + Map->IndexMask; OldEntry <= FinalEntry; ++OldEntry)
-    {
-        if (!OldEntry->OldFunction)
-        {
-            continue;
-        }
-
-        // TODO: If hash collisions are slow, we could have a secondary stronger hash here that
-		// is used on the second iteration followed by +1 for the subsequent iterations.
-        for (size_t Index = FunctionPtrHash(OldEntry->OldFunction, NewMap->Size2n);; Index = (Index + 1) & NewMap->IndexMask)
-        {
-            FFunctionEntry* NewEntry = NewMap->Entries + Index;
-            if (!NewEntry->OldFunction)
-            {
-                *NewEntry = *OldEntry;
-                break;
-            }
-        }
-    }
-    NewMap->KeyCount = Map->KeyCount;
-	free(Map);
-    return NewMap;
-}
+} // anonymous namespace
 
-static void FunctionMapAddImpl(void* OldFunction, void* NewFunction)
+void FunctionMapDumpStats()
 {
-    FFunctionMap* Map = GlobalData->FunctionMap;
-    if (AUTORTFM_UNLIKELY(!Map))
-    {
-		// We have no function map at all; create an empty map.
-        Map = AllocateFunctionMap(/*Size2n=*/ 10);
-        GlobalData->FunctionMap = Map;
-    }
-    else if (AUTORTFM_UNLIKELY(Map->KeyCount * 2 > Map->IndexMask))
-    {
-		// The function map is above 50% capacity; we need to grow and rehash.
-        Map = RehashFunctionMap(Map);
-        GlobalData->FunctionMap = Map;
-    }
-
-    for (size_t Index = FunctionPtrHash(OldFunction, Map->Size2n);; Index = (Index + 1) & Map->IndexMask)
-	{
-        FFunctionEntry* Entry = Map->Entries + Index;
-
-        if (!Entry->OldFunction) 
-		{
-            Entry->OldFunction = OldFunction;
-            Entry->NewFunction = NewFunction;
-            Map->KeyCount++;
-            return;
-        }
-
-        if (Entry->OldFunction == OldFunction) 
-		{
-            Entry->NewFunction = NewFunction;
-            return;
-        }
-    }
+	FFunctionMap::Get().DumpStats();
 }
 
-void FunctionMapAdd(void* OldFunction, void* NewFunction)
+void FunctionMapAdd(autortfm_open_to_closed_table* Tables)
 {
-	if (!OldFunction)
+	// Count the total number of new mappings before taking the lock
+	size_t NewMappingCount = 0;
+	for (const autortfm_open_to_closed_table* Table = Tables; Table; Table = Table->Next)
 	{
-		// Silently ignore null OldFunction to make weak linking easy.
-		return;
+		for (const autortfm_open_to_closed_mapping* Mapping = Table->Mappings; Mapping->Open; Mapping++)
+		{
+			NewMappingCount++;
+		}
 	}
 
-	InitializeGlobalDataIfNecessary();
+	FFunctionMap Map = FFunctionMap::Get();
+	Map.Reserve(NewMappingCount + Map.Count());
 
-	TLockGuard<FFastLock> LockGuard(GlobalData->FunctionMapLock);
-	FunctionMapAddImpl(OldFunction, NewFunction);
-	if (void* OldRelocated = FollowRelocation(OldFunction))
+	for (const autortfm_open_to_closed_table* Table = Tables; Table; Table = Table->Next)
 	{
-		FunctionMapAddImpl(OldRelocated, NewFunction);
-	}
-}
+		for (const autortfm_open_to_closed_mapping* Mapping = Table->Mappings; Mapping->Open; Mapping++)
+		{
+			AUTORTFM_VERBOSE("Registering open %p -> %p", Mapping->Open, Mapping->Closed);
 
-static void* FunctionMapReportError(void* OldFunction, const char* Where)
-{
-    std::string FunctionDescription = GetFunctionDescription(OldFunction);
-	if (Where)
-	{
-		AUTORTFM_REPORT_ERROR("Could not find function %p '%s' where '%s'.", OldFunction, FunctionDescription.c_str(), Where);
-	}
-	else
-	{
-		AUTORTFM_REPORT_ERROR("Could not find function %p '%s'.", OldFunction, FunctionDescription.c_str());
+			Map.Add(Mapping->Open, Mapping->Closed);
+
+			if (void* OpenRelocated = FollowRelocation(Mapping->Open))
+			{
+				Map.Add(OpenRelocated, Mapping->Closed);
+			}
+		}
 	}
-	return nullptr;
 }
 
-void* FunctionMapLookupTryFollowingRelocation(void* Function, const char* Where)
+void* FunctionMapLookupExhaustive(void* OpenFn, const char* Where)
 {
-	if (void* Ptr = FollowRelocation(Function); AUTORTFM_LIKELY(Ptr))
+	// Use an explicit scope for Map as FunctionMapReportError() may unwind the
+	// stack without first calling the destructor of Map, leaving the
+	// FFunctionMap locked.
 	{
-		AUTORTFM_MUST_TAIL return FunctionMapLookup(Ptr, Where);
-	}
-	AUTORTFM_MUST_TAIL return FunctionMapReportError(Function, Where);
-}
+		FFunctionMap Map = FFunctionMap::Get();
 
-void FunctionMapDumpStats()
-{
-    const FFunctionMap* const Map = GlobalData->FunctionMap;
-	size_t MapSlots = 1llu << Map->Size2n;
-	AUTORTFM_LOG("Function Map Stats");
-	AUTORTFM_LOG("==================");
-	AUTORTFM_LOG("Occupancy: %zu entries in %zu slots (load factor %g%%)", 
-		         Map->KeyCount, 
-		         MapSlots,
-				 double(Map->KeyCount * 1000llu / MapSlots) / 10.0);
-
-	std::map<uintptr_t, size_t> CollisionMap; // <probes required to find the function, number of functions>
-	for (const FFunctionEntry* Entry = Map->Entries, *FinalEntry = Map->Entries + Map->IndexMask; Entry <= FinalEntry; ++Entry)
-	{
-		if (Entry->OldFunction)
+		if (void* ClosedFn = Map.Lookup(OpenFn); AUTORTFM_LIKELY(ClosedFn))
 		{
-			uintptr_t ActualIndex = std::distance(Map->Entries, Entry);
-			uintptr_t IdealIndex = FunctionPtrHash(Entry->OldFunction, Map->Size2n) & Map->IndexMask;
-			// We need to account for wraparound in this calculation.
-			uintptr_t Delta = (IdealIndex <= ActualIndex) 
-				? ActualIndex - IdealIndex
-				: ActualIndex + MapSlots - IdealIndex;
-				                                           
-			CollisionMap[Delta] += 1;
+			return ClosedFn;
 		}
-	}
 
-	for (uintptr_t NumProbes = 0, HighestProbes = CollisionMap.rbegin()->first; NumProbes <= HighestProbes; ++NumProbes)
-	{
-		AUTORTFM_LOG("%2" PRIuPTR " probes: %zu functions", NumProbes, CollisionMap[NumProbes]); 
+		if (void* RelocatedOpenFn = FollowRelocation(OpenFn); AUTORTFM_LIKELY(RelocatedOpenFn))
+		{
+			if (void* ClosedFn = FunctionMapLookupUsingMagicPrefix(RelocatedOpenFn); AUTORTFM_LIKELY(ClosedFn))
+			{
+				return ClosedFn;
+			}
+			if (void* ClosedFn = Map.Lookup(RelocatedOpenFn); AUTORTFM_LIKELY(ClosedFn))
+			{
+				return ClosedFn;
+			}
+		}
 	}
+
+	AUTORTFM_MUST_TAIL return FunctionMapReportError(OpenFn, Where);
 }
 
 } // namespace AutoRTFM
diff --git a/Engine/Source/Runtime/AutoRTFM/Private/FunctionMap.h b/Engine/Source/Runtime/AutoRTFM/Private/FunctionMap.h
--- a/Engine/Source/Runtime/AutoRTFM/Private/FunctionMap.h
+++ b/Engine/Source/Runtime/AutoRTFM/Private/FunctionMap.h
@@ -4,75 +4,25 @@
 
 #if (defined(__AUTORTFM) && __AUTORTFM)
 
-#include "GlobalData.h"
-#include "Utils.h"
+#include "AutoRTFM.h"
 
 namespace AutoRTFM
 {
 
-inline uintptr_t FunctionPtrHash(void* FunctionPtr, size_t HashTableSize2n)
-{
-	// Use a Fibonacci hash to fold our 64-bit value into a hash-appropriate value.
-	// https://probablydance.com/2018/06/16/fibonacci-hashing-the-optimization-that-the-world-forgot-or-a-better-alternative-to-integer-modulo/
-	uintptr_t HashBits = reinterpret_cast<uintptr_t>(FunctionPtr);
-	// Apply Fibonacci product and preserve the highest bits.
-	return (HashBits * 11400714819323198485llu) >> (64 - HashTableSize2n);
-}
-
-struct FFunctionEntry
-{
-    void* OldFunction = nullptr;
-	void* NewFunction = nullptr;
-};
-
-struct FFunctionMap 
-{
-	size_t KeyCount;           // The number of elements in the map.
-	size_t Size2n;             // The map size as a power-of-two; that is, a Size2n of 10 corresponds to 1024 entries, as 2^10 == 1024.
-	size_t IndexMask;          // The map size minus one, for use as a bitmask; a function map with 1024 elements will have an IndexMask of 0x3FF.
-    FFunctionEntry Entries[1]; // Actually contains `2 ^ Size2n` elements.
-};
-
 // Dumps statistics about the function map to LOG_INFO.
 void FunctionMapDumpStats();
 
-void FunctionMapAdd(void* OldFunction, void* NewFunction);
+// Adds all the open to closed function mappings from the linked list of tables to the function map.
+void FunctionMapAdd(autortfm_open_to_closed_table* Tables);
 
-// Attempts to lookup the function through a relocation thunk.
-// This is called by FunctionMapLookup(), and should not be called directly.
-void* FunctionMapLookupTryFollowingRelocation(void* Function, const char* Where);
-
-inline void* FunctionMapTryLookup(void* OldFunction)
-{
-    FFunctionMap* Map = GlobalData->FunctionMap;
-
-    if (AUTORTFM_UNLIKELY(!Map))
-	{
-        return nullptr;
-	}
-
-    for (size_t Index = FunctionPtrHash(OldFunction, Map->Size2n);; Index = (Index + 1) & Map->IndexMask) 
-	{
-        FFunctionEntry* Entry = Map->Entries + Index;
-
-        if (Entry->OldFunction == OldFunction)
-        {
-            // This races with Add, but in a benign way. If _OldFunction is set and matches our OldFunction but the _NewFunction
-            // is nullptr, then it means that an Add is happening right now, and before the Add, this entry would have been nullptr.
-            return Entry->NewFunction;
-        }
-
-        if (!Entry->OldFunction)
-        {
-            return nullptr;
-        }
-    }
-}
+// Looks up the closed function from the open function using the function map, following dynamic
+// library thunk redirections.
+void* FunctionMapLookupExhaustive(void* OpenFn, const char* Where);
 
-template<typename TReturnType, typename... TParameterTypes>
-auto FunctionMapTryLookup(TReturnType (*Function)(TParameterTypes...)) -> TReturnType (*)(TParameterTypes...)
+template<typename ReturnType, typename... ParameterTypes>
+auto FunctionMapLookupExhaustive(ReturnType (*Function)(ParameterTypes...), const char* Where) -> ReturnType (*)(ParameterTypes...)
 {
-    return reinterpret_cast<TReturnType (*)(TParameterTypes...)>(FunctionMapTryLookup(reinterpret_cast<void*>(Function)));
+	return reinterpret_cast<ReturnType (*)(ParameterTypes...)>(FunctionMapLookupExhaustive(reinterpret_cast<void*>(Function), Where));
 }
 
 } // namespace AutoRTFM
diff --git a/Engine/Source/Runtime/AutoRTFM/Private/FunctionMapInlines.h b/Engine/Source/Runtime/AutoRTFM/Private/FunctionMapInlines.h
--- a/Engine/Source/Runtime/AutoRTFM/Private/FunctionMapInlines.h
+++ b/Engine/Source/Runtime/AutoRTFM/Private/FunctionMapInlines.h
@@ -11,20 +11,8 @@
 namespace AutoRTFM
 {
 
-inline void* FunctionMapLookupSlow(void* OldFunction, const char* Where)
-{
-	// Instead fall back to the slower function map lookup.
-    void* const Result = FunctionMapTryLookup(OldFunction);
-
-	if (AUTORTFM_UNLIKELY(!Result))
-	{
-		AUTORTFM_MUST_TAIL return FunctionMapLookupTryFollowingRelocation(OldFunction, Where);
-	}
-
-	return Result;
-}
-
-inline void* FunctionMapLookup(void* OldFunction, const char* Where)
+UE_AUTORTFM_FORCEINLINE
+void* FunctionMapLookupUsingMagicPrefix(void* OpenFn)
 {
 	// We use prefix data in our custom LLVM pass to stuff some data just
 	// before the address of all open function pointers (that we have
@@ -33,7 +21,7 @@ inline void* FunctionMapLookup(void* OldFunction, const char* Where)
 	// to give us a much higher confidence that there is actually a closed
 	// variant pointer residing 8-bytes before our function address.
 	uint64_t PrefixData;
-	memcpy(&PrefixData, reinterpret_cast<char*>(OldFunction) - sizeof(uint64_t), sizeof(uint64_t));
+	memcpy(&PrefixData, reinterpret_cast<char*>(OpenFn) - sizeof(uint64_t), sizeof(uint64_t));
 	if (AUTORTFM_LIKELY(Constants::MagicPrefix == (PrefixData & 0xffff'0000'0000'0000)))
 	{
 		return reinterpret_cast<void*>(PrefixData & 0x0000'ffff'ffff'ffff);
@@ -44,20 +32,30 @@ inline void* FunctionMapLookup(void* OldFunction, const char* Where)
 	// hash. If we see it, check the preceding 64-bit word for our prefix.
 	else if (AUTORTFM_UNLIKELY(0xc105'cafe == (PrefixData & 0x0000'0000'ffff'ffff)))
 	{
-		memcpy(&PrefixData, reinterpret_cast<char*>(OldFunction) - sizeof(uint64_t) * 2, sizeof(uint64_t));
+		memcpy(&PrefixData, reinterpret_cast<char*>(OpenFn) - sizeof(uint64_t) * 2, sizeof(uint64_t));
 		if (AUTORTFM_LIKELY(Constants::MagicPrefix == (PrefixData & 0xffff'0000'0000'0000)))
 		{
 			return reinterpret_cast<void*>(PrefixData & 0x0000'ffff'ffff'ffff);
 		}
 	}
 
-	AUTORTFM_MUST_TAIL return FunctionMapLookupSlow(OldFunction, Where);
+	return nullptr;
+}
+
+inline void* FunctionMapLookup(void* OpenFn, const char* Where)
+{
+	if (void* ClosedFn = FunctionMapLookupUsingMagicPrefix(OpenFn); AUTORTFM_LIKELY(ClosedFn))
+	{
+		return ClosedFn;
+	}
+
+	AUTORTFM_MUST_TAIL return FunctionMapLookupExhaustive(OpenFn, Where);
 }
 
 template<typename TReturnType, typename... TParameterTypes>
-auto FunctionMapLookup(TReturnType (*Function)(TParameterTypes...), const char* Where) -> TReturnType (*)(TParameterTypes...)
+auto FunctionMapLookup(TReturnType (*OpenFn)(TParameterTypes...), const char* Where) -> TReturnType (*)(TParameterTypes...)
 {
-    return reinterpret_cast<TReturnType (*)(TParameterTypes...)>(FunctionMapLookup(reinterpret_cast<void*>(Function), Where));
+    return reinterpret_cast<TReturnType (*)(TParameterTypes...)>(FunctionMapLookup(reinterpret_cast<void*>(OpenFn), Where));
 }
 
 } // namespace AutoRTFM
diff --git a/Engine/Source/Runtime/AutoRTFM/Private/GlobalData.cpp b/Engine/Source/Runtime/AutoRTFM/Private/GlobalData.cpp
deleted file mode 100644
diff --git a/Engine/Source/Runtime/AutoRTFM/Private/GlobalData.h b/Engine/Source/Runtime/AutoRTFM/Private/GlobalData.h
deleted file mode 100644
diff --git a/Engine/Source/Runtime/AutoRTFM/Private/LockGuard.h b/Engine/Source/Runtime/AutoRTFM/Private/LockGuard.h
deleted file mode 100644
diff --git a/Engine/Source/Runtime/AutoRTFM/Private/Transaction.cpp b/Engine/Source/Runtime/AutoRTFM/Private/Transaction.cpp
--- a/Engine/Source/Runtime/AutoRTFM/Private/Transaction.cpp
+++ b/Engine/Source/Runtime/AutoRTFM/Private/Transaction.cpp
@@ -7,7 +7,6 @@
 
 #include "CallNestInlines.h"
 #include "ContextStatus.h"
-#include "GlobalData.h"
 #include "HitSet.h"
 #include "Utils.h"
 
