[AutoRTFM] Allow UE_AUTORTFM_NOAUTORTFM functions to be inlined

The reason this wasn't permitted before is that the function could be inlined before running the AutoRTFM pass that forks the closed variant. Because the function is inlined, the pass cannot tell the inlined code should not exist in the closed variant, and the no-autortfm attribute is ignored.

This generalizes the AutoRTFMStaticInitOutlinerPass pass into a AutoRTFMEarlyPass, which does what it did before as well as now injecting the following pseudo-code at the top of functions annotated with [[clang::noautortfm]]:
if (autortfm_is_closed()) { autortfm_called_no_autortfm(); __builtin_unreachable(); }

Now, if the annotated function is inlined into the closed function, it'll error.
Because the call to autortfm_is_closed() is replaced with a constant boolean, this conditional branch is optimized away from both the open and closed variants.

Other changes:
* Extract out the llvm::Type helpers from AutoRTFMInstrumentationPassImpl into a new BasicTypes helper, used by both AutoRTFMEarlyPassImpl and AutoRTFMInstrumentationPassImpl
* Extract out the imported functions lazy-cache from AutoRTFMInstrumentationPassImpl into a new ImportedFunctions helper, used by both AutoRTFMEarlyPassImpl and AutoRTFMInstrumentationPassImpl. This drops the use of std::map for explicit functions.
* All AutoSTM code has been removed from the compiler to reduce maintainance burden. This experimental code has no runtime and is extremely unlikely to be used. We can use version control in the unlikely chance we want to bring it back.
* Remove the AutoRTFMAttribute.cpp lit test. This was testing [[clang::noautortfm]], which is already tested in NoAutoRTFMAttribute.cpp.
* NoAutoRTFMAttribute.cpp lit tests have been expanded to cover inlining cases.

diff --git a/Engine/Source/Runtime/AutoRTFM/Private/CompilerInterface.cpp b/Engine/Source/Runtime/AutoRTFM/Private/CompilerInterface.cpp
--- a/Engine/Source/Runtime/AutoRTFM/Private/CompilerInterface.cpp
+++ b/Engine/Source/Runtime/AutoRTFM/Private/CompilerInterface.cpp
@@ -168,21 +168,9 @@ extern "C" UE_AUTORTFM_API void autortfm_llvm_missing_function()
     Context->AbortByLanguageAndThrow();
 }
 
-extern "C" UE_AUTORTFM_API void autortfm_llvm_alignment_error(FContext* Context, void* Ptr, size_t Alignment, const char* Message)
+extern "C" UE_AUTORTFM_API void autortfm_called_no_autortfm()
 {
-    AbortDueToBadAlignment(Context, Ptr, Alignment, Message);
-}
-
-extern "C" UE_AUTORTFM_API void autortfm_llvm_error(const char* Message)
-{
-	if (Message)
-	{
-		AUTORTFM_FATAL("Transaction failing because of LLVM issue '%s'", Message);
-	}
-	else
-	{
-		AUTORTFM_FATAL("Transaction failing because of LLVM issue");
-	}
+	AUTORTFM_FATAL("inlined UE_AUTORTFM_NOAUTORTFM function called from the closed");
 }
 
 } // namespace AutoRTFM
diff --git a/Engine/Source/Runtime/AutoRTFM/Private/Context.cpp b/Engine/Source/Runtime/AutoRTFM/Private/Context.cpp
--- a/Engine/Source/Runtime/AutoRTFM/Private/Context.cpp
+++ b/Engine/Source/Runtime/AutoRTFM/Private/Context.cpp
@@ -412,12 +412,6 @@ AutoRTFM::FStackRange FContext::GetThreadStackRange()
 
 ETransactionResult FContext::Transact(void (*UninstrumentedFunction)(void*), void (*InstrumentedFunction)(void*), void* Arg)
 {
-	if constexpr (0 != UE_AUTOSTM)
-	{
-		AUTORTFM_WARN("AutoSTM is not implemented in the runtime yet!");
-		return ETransactionResult::AbortedByLanguage;
-	}
-
     if (AUTORTFM_UNLIKELY(EContextStatus::Committing == Status))
     {
     	return ETransactionResult::AbortedByTransactInOnCommit;
diff --git a/Engine/Source/Runtime/AutoRTFM/Public/AutoRTFM.h b/Engine/Source/Runtime/AutoRTFM/Public/AutoRTFM.h
--- a/Engine/Source/Runtime/AutoRTFM/Public/AutoRTFM.h
+++ b/Engine/Source/Runtime/AutoRTFM/Public/AutoRTFM.h
@@ -32,18 +32,6 @@
 #define UE_AUTORTFM_ENABLED_RUNTIME_BY_DEFAULT 1
 #endif
 
-#if UE_AUTORTFM
-#if defined(__has_feature)
-#if __has_feature(autostm)
-#define UE_AUTOSTM 1
-#endif
-#endif
-#endif
-
-#if !defined(UE_AUTOSTM)
-#define UE_AUTOSTM 0
-#endif
-
 #if !defined(UE_AUTORTFM_STATIC_VERIFIER)
 #define UE_AUTORTFM_STATIC_VERIFIER 0
 #endif
