[AutoRTFM] Fix unused variable warnings / errors when building with UE_AUTORTFM disabled

diff --git a/Engine/Source/Runtime/AutoRTFM/Private/API.cpp b/Engine/Source/Runtime/AutoRTFM/Private/API.cpp
--- a/Engine/Source/Runtime/AutoRTFM/Private/API.cpp
+++ b/Engine/Source/Runtime/AutoRTFM/Private/API.cpp
@@ -11,7 +11,6 @@
 
 #if UE_AUTORTFM
 static_assert(UE_AUTORTFM_ENABLED, "AutoRTFM/API.cpp requires the compiler flag '-fautortfm'");
-#endif
 
 namespace
 {
@@ -47,6 +46,7 @@ namespace
 	// is called.
 	autortfm_open_to_closed_table* GPendingOpenToClosedRegistrations = nullptr;
 }
+#endif  // UE_AUTORTFM
 
 namespace AutoRTFM
 {
@@ -54,18 +54,21 @@ namespace AutoRTFM
 	{
 		UE_AUTORTFM_API ForTheRuntime::EAutoRTFMEnabledState ForceSetAutoRTFMRuntime(ForTheRuntime::EAutoRTFMEnabledState State)
 		{
+#if UE_AUTORTFM
 			const int Original = GAutoRTFMRuntimeEnabled;
 			if (GAutoRTFMRuntimeEnabled != State)
 			{
 				GAutoRTFMRuntimeEnabled = State;
-#if UE_AUTORTFM
 				if (GExternAPI.OnRuntimeEnabledChanged)
 				{
 					GExternAPI.OnRuntimeEnabledChanged();
 				}
-#endif
 			}
 			return static_cast<ForTheRuntime::EAutoRTFMEnabledState>(Original);
+#else
+			UE_AUTORTFM_UNUSED(State);
+			return ForTheRuntime::AutoRTFM_ForcedDisabled;
+#endif
 		}
 	} // Testing
 } // AutoRTFM
@@ -220,22 +223,38 @@ namespace AutoRTFM
 
 		void SetInternalAbortAction(EAutoRTFMInternalAbortActionState State)
 		{
+#if UE_AUTORTFM
 			GAutoRTFMInternalAbortAction = State;
+#else
+			UE_AUTORTFM_UNUSED(State);
+#endif
 		}
 
 		EAutoRTFMInternalAbortActionState GetInternalAbortAction()
 		{
+#if UE_AUTORTFM
 			return static_cast<EAutoRTFMInternalAbortActionState>(GAutoRTFMInternalAbortAction);
+#else
+			return Crash;
+#endif
 		}
 
 		bool GetEnsureOnInternalAbort()
 		{
+#if UE_AUTORTFM
 			return GAutoRTFMEnsureOnInternalAbort;
+#else
+			return false;
+#endif
 		}
 
 		void SetEnsureOnInternalAbort(bool bEnabled)
 		{
+#if UE_AUTORTFM
 			GAutoRTFMEnsureOnInternalAbort = bEnabled;
+#else
+			UE_AUTORTFM_UNUSED(bEnabled);
+#endif
 		}
 
 		void SetRetryTransaction(EAutoRTFMRetryTransactionState State)
