[AutoRTFM] Add support for AutoRTFM + TSAN

Enabling TSAN wouldn't link as the TSAN libraries declare global operator overloads for new and delete which collide with UE's overloads. TSAN works without AutoRTFM enabled as UBT will enable `FORCE_ANSI_ALLOCATOR` when a sanitizer is enabled, which forces `new` and `delete` to use the standard library's implementation instead of the custom allocator's version.
CL 27583795 explicitly added an exception for `FORCE_ANSI_ALLOCATOR` when building with AutoRTFM with this comment:

// For AutoRTFM when FORCE_ANSI_ALLOCATOR is specified we still need to re-route operator new/delete (to deal
// with transactionalization of memory allocation). Address sanitizer will still work even with this re-routing
// because we point the underlying FMemory allocator at ansi malloc/free, which ASan still hijacks for its
// purposes.

This CL removes the AutoRTFM `FORCE_ANSI_ALLOCATOR` exception and adds open -> closed mappings for all the `new` and `delete` overloads that correctly handles the allocation / free in a transaction.

diff --git a/Engine/Source/Runtime/AutoRTFM/Private/StdLibSupport.cpp b/Engine/Source/Runtime/AutoRTFM/Private/StdLibSupport.cpp
--- a/Engine/Source/Runtime/AutoRTFM/Private/StdLibSupport.cpp
+++ b/Engine/Source/Runtime/AutoRTFM/Private/StdLibSupport.cpp
@@ -14,6 +14,7 @@
 #include <functional> // note: introduces additional math overloads
 #include <stdarg.h>
 #include <stdio.h>
+#include <stdlib.h>
 #include <string.h>
 #include <wchar.h>
 
@@ -610,6 +611,142 @@ extern "C"
 }
 #endif
 
+#if FORCE_ANSI_ALLOCATOR
+template<typename OpenNewFn, typename OpenDeleteFn, typename ... ArgTys>
+UE_AUTORTFM_FORCEINLINE static void* RTFM_ClosedNew(OpenNewFn* New, OpenDeleteFn* Delete, unsigned long Size, ArgTys... Args)
+{
+	void* Result = New(Size, Args...);
+	FContext* Context = FContext::Get();
+	Context->GetCurrentTransaction()->DeferUntilAbort([=]
+	{
+		Delete(Result, Args...);
+	});
+	Context->DidAllocate(Result, Size);
+	return Result;
+}
+
+template<typename OpenDeleteFn, typename ... ArgTys>
+UE_AUTORTFM_FORCEINLINE static void RTFM_ClosedDelete(OpenDeleteFn* Delete, void* Pointer, ArgTys... Args)
+{
+	if (Pointer)
+	{
+		FContext* Context = FContext::Get();
+		Context->GetCurrentTransaction()->DeferUntilCommit([=]
+		{
+			Delete(Pointer, Args...);
+		});
+	}
+}
+
+extern "C"
+{
+	// Operator new and new[]
+	void* _Znwm(unsigned long);
+	void* _Znam(unsigned long);
+	void* _ZnwmRKSt9nothrow_t(unsigned long, const std::nothrow_t&);
+	void* _ZnamRKSt9nothrow_t(unsigned long, const std::nothrow_t&);
+	void* _ZnwmSt11align_val_t(unsigned long, std::align_val_t);
+	void* _ZnamSt11align_val_t(unsigned long, std::align_val_t);
+	void* _ZnwmSt11align_val_tRKSt9nothrow_t(unsigned long, std::align_val_t, const std::nothrow_t&);
+	void* _ZnamSt11align_val_tRKSt9nothrow_t(unsigned long, std::align_val_t, const std::nothrow_t&);
+	// Operator delete and delete[]
+	void _ZdlPv(void*);
+	void _ZdaPv(void*);
+	void _ZdlPvRKSt9nothrow_t(void*, const std::nothrow_t&);
+	void _ZdaPvRKSt9nothrow_t(void*, const std::nothrow_t&);
+	void _ZdlPvm(void*, unsigned long);
+	void _ZdaPvm(void*, unsigned long);
+	void _ZdlPvSt11align_val_t(void*, std::align_val_t);
+	void _ZdaPvSt11align_val_t(void*, std::align_val_t);
+	void _ZdlPvSt11align_val_tRKSt9nothrow_t(void*, std::align_val_t, const std::nothrow_t&);
+	void _ZdaPvSt11align_val_tRKSt9nothrow_t(void*, std::align_val_t, const std::nothrow_t&);
+	void _ZdlPvmSt11align_val_t(void*, unsigned long, std::align_val_t);
+	void _ZdaPvmSt11align_val_t(void*, unsigned long, std::align_val_t);
+
+	void* RTFM_Znwm(unsigned long Size)
+	{
+		return RTFM_ClosedNew(_Znwm, _ZdlPv, Size);
+	}
+	void* RTFM_Znam(unsigned long Size)
+	{
+		return RTFM_ClosedNew(_Znam, _ZdaPv, Size);
+	}
+	void* RTFM_ZnwmRKSt9nothrow_t(unsigned long Size, const std::nothrow_t& NoThrow)
+	{
+		return RTFM_ClosedNew(_ZnwmRKSt9nothrow_t, _ZdlPvRKSt9nothrow_t, Size, NoThrow);
+	}
+	void* RTFM_ZnamRKSt9nothrow_t(unsigned long Size, const std::nothrow_t& NoThrow)
+	{
+		return RTFM_ClosedNew(_ZnamRKSt9nothrow_t, _ZdaPvRKSt9nothrow_t, Size, NoThrow);
+	}
+	void* RTFM_ZnwmSt11align_val_t(unsigned long Size, std::align_val_t Align)
+	{
+		return RTFM_ClosedNew(_ZnwmSt11align_val_t, _ZdlPvSt11align_val_t, Size, Align);
+	}
+	void* RTFM_ZnamSt11align_val_t(unsigned long Size, std::align_val_t Align)
+	{
+		return RTFM_ClosedNew(_ZnamSt11align_val_t, _ZdaPvSt11align_val_t, Size, Align);
+	}
+	void* RTFM_ZnwmSt11align_val_tRKSt9nothrow_t(unsigned long Size, std::align_val_t Align, const std::nothrow_t& NoThrow)
+	{
+		return RTFM_ClosedNew(_ZnwmSt11align_val_tRKSt9nothrow_t, _ZdlPvSt11align_val_tRKSt9nothrow_t, Size, Align, NoThrow);
+	}
+	void* RTFM_ZnamSt11align_val_tRKSt9nothrow_t(unsigned long Size, std::align_val_t Align, const std::nothrow_t& NoThrow)
+	{
+		return RTFM_ClosedNew(_ZnamSt11align_val_tRKSt9nothrow_t, _ZdaPvSt11align_val_tRKSt9nothrow_t, Size, Align, NoThrow);
+	}
+
+	void RTFM_ZdlPv(void* Pointer)
+	{
+		RTFM_ClosedDelete(_ZdlPv, Pointer);
+	}
+	void RTFM_ZdaPv(void* Pointer)
+	{
+		RTFM_ClosedDelete(_ZdaPv, Pointer);
+	}
+	void RTFM_ZdlPvRKSt9nothrow_t(void* Pointer, const std::nothrow_t& NoThrow)
+	{
+		RTFM_ClosedDelete(_ZdlPvRKSt9nothrow_t, Pointer, NoThrow);
+	}
+	void RTFM_ZdaPvRKSt9nothrow_t(void* Pointer, const std::nothrow_t& NoThrow)
+	{
+		RTFM_ClosedDelete(_ZdaPvRKSt9nothrow_t, Pointer, NoThrow);
+	}
+	void RTFM_ZdlPvm(void* Pointer, unsigned long Size)
+	{
+		RTFM_ClosedDelete(_ZdlPvm, Pointer, Size);
+	}
+	void RTFM_ZdaPvm(void* Pointer, unsigned long Size)
+	{
+		RTFM_ClosedDelete(_ZdaPvm, Pointer, Size);
+	}
+	void RTFM_ZdlPvSt11align_val_t(void* Pointer, std::align_val_t Align)
+	{
+		RTFM_ClosedDelete(_ZdlPvSt11align_val_t, Pointer, Align);
+	}
+	void RTFM_ZdaPvSt11align_val_t(void* Pointer, std::align_val_t Align)
+	{
+		RTFM_ClosedDelete(_ZdaPvSt11align_val_t, Pointer, Align);
+	}
+	void RTFM_ZdlPvSt11align_val_tRKSt9nothrow_t(void* Pointer, std::align_val_t Align, const std::nothrow_t& NoThrow)
+	{
+		RTFM_ClosedDelete(_ZdlPvSt11align_val_tRKSt9nothrow_t, Pointer, Align, NoThrow);
+	}
+	void RTFM_ZdaPvSt11align_val_tRKSt9nothrow_t(void* Pointer, std::align_val_t Align, const std::nothrow_t& NoThrow)
+	{
+		RTFM_ClosedDelete(_ZdaPvSt11align_val_tRKSt9nothrow_t, Pointer, Align, NoThrow);
+	}
+	void RTFM_ZdlPvmSt11align_val_t(void* Pointer, unsigned long Size, std::align_val_t Align)
+	{
+		RTFM_ClosedDelete(_ZdlPvmSt11align_val_t, Pointer, Size, Align);
+	}
+	void RTFM_ZdaPvmSt11align_val_t(void* Pointer, unsigned long Size, std::align_val_t Align)
+	{
+		RTFM_ClosedDelete(_ZdaPvmSt11align_val_t, Pointer, Size, Align);
+	}
+}
+#endif
+
 UE_AUTORTFM_REGISTER_OPEN_TO_CLOSED_FUNCTIONS(
 #if AUTORTFM_ASAN_ENABLED
 	UE_AUTORTFM_MAP_OPEN_TO_SELF(__asan_addr_is_in_fake_stack),
@@ -699,6 +836,29 @@ UE_AUTORTFM_REGISTER_OPEN_TO_CLOSED_FUNCTIONS(
 	UE_AUTORTFM_MAP_OPEN_TO_SELF_OVERLOADED(const char* (const char*, const char*), strstr),
 #endif
 
+#if FORCE_ANSI_ALLOCATOR
+	UE_AUTORTFM_MAP_OPEN_TO_CLOSED(_Znwm, RTFM_Znwm),
+	UE_AUTORTFM_MAP_OPEN_TO_CLOSED(_Znam, RTFM_Znam),
+	UE_AUTORTFM_MAP_OPEN_TO_CLOSED(_ZnwmRKSt9nothrow_t, RTFM_ZnwmRKSt9nothrow_t),
+	UE_AUTORTFM_MAP_OPEN_TO_CLOSED(_ZnamRKSt9nothrow_t, RTFM_ZnamRKSt9nothrow_t),
+	UE_AUTORTFM_MAP_OPEN_TO_CLOSED(_ZnwmSt11align_val_t, RTFM_ZnwmSt11align_val_t),
+	UE_AUTORTFM_MAP_OPEN_TO_CLOSED(_ZnamSt11align_val_t, RTFM_ZnamSt11align_val_t),
+	UE_AUTORTFM_MAP_OPEN_TO_CLOSED(_ZnwmSt11align_val_tRKSt9nothrow_t, RTFM_ZnwmSt11align_val_tRKSt9nothrow_t),
+	UE_AUTORTFM_MAP_OPEN_TO_CLOSED(_ZnamSt11align_val_tRKSt9nothrow_t, RTFM_ZnamSt11align_val_tRKSt9nothrow_t),
+	UE_AUTORTFM_MAP_OPEN_TO_CLOSED(_ZdlPv, RTFM_ZdlPv),
+	UE_AUTORTFM_MAP_OPEN_TO_CLOSED(_ZdaPv, RTFM_ZdaPv),
+	UE_AUTORTFM_MAP_OPEN_TO_CLOSED(_ZdlPvRKSt9nothrow_t, RTFM_ZdlPvRKSt9nothrow_t),
+	UE_AUTORTFM_MAP_OPEN_TO_CLOSED(_ZdaPvRKSt9nothrow_t, RTFM_ZdaPvRKSt9nothrow_t),
+	UE_AUTORTFM_MAP_OPEN_TO_CLOSED(_ZdlPvm, RTFM_ZdlPvm),
+	UE_AUTORTFM_MAP_OPEN_TO_CLOSED(_ZdaPvm, RTFM_ZdaPvm),
+	UE_AUTORTFM_MAP_OPEN_TO_CLOSED(_ZdlPvSt11align_val_t, RTFM_ZdlPvSt11align_val_t),
+	UE_AUTORTFM_MAP_OPEN_TO_CLOSED(_ZdaPvSt11align_val_t, RTFM_ZdaPvSt11align_val_t),
+	UE_AUTORTFM_MAP_OPEN_TO_CLOSED(_ZdlPvSt11align_val_tRKSt9nothrow_t, RTFM_ZdlPvSt11align_val_tRKSt9nothrow_t),
+	UE_AUTORTFM_MAP_OPEN_TO_CLOSED(_ZdaPvSt11align_val_tRKSt9nothrow_t, RTFM_ZdaPvSt11align_val_tRKSt9nothrow_t),
+	UE_AUTORTFM_MAP_OPEN_TO_CLOSED(_ZdlPvmSt11align_val_t, RTFM_ZdlPvmSt11align_val_t),
+	UE_AUTORTFM_MAP_OPEN_TO_CLOSED(_ZdaPvmSt11align_val_t, RTFM_ZdaPvmSt11align_val_t),
+#endif
+
 	UE_AUTORTFM_MAP_OPEN_TO_CLOSED(memcpy, RTFM_memcpy),
 	UE_AUTORTFM_MAP_OPEN_TO_CLOSED(memmove, RTFM_memmove),
 	UE_AUTORTFM_MAP_OPEN_TO_CLOSED(memset, RTFM_memset),
diff --git a/Engine/Source/Runtime/Core/Public/Modules/Boilerplate/ModuleBoilerplate.h b/Engine/Source/Runtime/Core/Public/Modules/Boilerplate/ModuleBoilerplate.h
--- a/Engine/Source/Runtime/Core/Public/Modules/Boilerplate/ModuleBoilerplate.h
+++ b/Engine/Source/Runtime/Core/Public/Modules/Boilerplate/ModuleBoilerplate.h
@@ -40,11 +40,7 @@
 
 // Disable the replacement new/delete when running the Clang static analyzer, due to false positives in 15.0.x:
 // [REDACTED_PROFILE_URL]/llvm-project/issues/58820
-// For AutoRTFM when FORCE_ANSI_ALLOCATOR is specified we still need to re-route operator new/delete (to deal
-// with transactionalization of memory allocation). Address sanitizer will still work even with this re-routing
-// because we point the underlying FMemory allocator at ansi malloc/free, which ASan still hijacks for its
-// purposes.
-#if !(FORCE_ANSI_ALLOCATOR && !(defined(__AUTORTFM) && __AUTORTFM)) && !defined(__clang_analyzer__)
+#if !FORCE_ANSI_ALLOCATOR && !defined(__clang_analyzer__)
 	static_assert(__STDCPP_DEFAULT_NEW_ALIGNMENT__ <= 16, "Expecting 16-byte default operator new alignment - alignments > 16 may have bloat");
 	#define REPLACEMENT_OPERATOR_NEW_AND_DELETE \
 		OPERATOR_NEW_MSVC_PRAGMA void* operator new  ( size_t Size                                                    ) OPERATOR_NEW_THROW_SPEC      { return FMemory::Malloc( Size ? Size : 1, __STDCPP_DEFAULT_NEW_ALIGNMENT__ ); } \
