[AutoRTFM] Optimization: Don't lookup autortfm_open[_explicit_validation] from the closed

Give the compiler the smarts to know what function this maps to.
Provides benchmarked performance gains of up to 3%, and shaves ~100k off server binaries.

diff --git a/Engine/Source/Runtime/AutoRTFM/Private/API.cpp b/Engine/Source/Runtime/AutoRTFM/Private/API.cpp
--- a/Engine/Source/Runtime/AutoRTFM/Private/API.cpp
+++ b/Engine/Source/Runtime/AutoRTFM/Private/API.cpp
@@ -546,7 +546,7 @@ extern "C" UE_AUTORTFM_NOAUTORTFM UE_AUTORTFM_API void autortfm_post_static_loca
 	}
 }
 
-AUTORTFM_CORE_API void RTFM_autortfm_open(autortfm_memory_validation_level MemoryValidationLevel, void (*Work)(void*), void* Arg, const void* ReturnAddress)
+extern "C" AUTORTFM_CORE_API UE_AUTORTFM_API void autortfm_open_from_closed_explicit_validation(autortfm_memory_validation_level MemoryValidationLevel, void (*Work)(void*), void* Arg, const void* ReturnAddress)
 {
 	FTransaction* Transaction = FContext::Get()->GetCurrentTransaction();
 	if (Transaction)
@@ -572,6 +572,28 @@ AUTORTFM_CORE_API void RTFM_autortfm_open(autortfm_memory_validation_level Memor
 	}
 }
 
+extern "C" AUTORTFM_CORE_API UE_AUTORTFM_API void autortfm_open_from_closed(void (*Work)(void*), void* Arg, const void* ReturnAddress)
+{
+	FTransaction* Transaction = FContext::Get()->GetCurrentTransaction();
+	if (Transaction)
+	{
+		Transaction->SetOpenActive(ForTheRuntime::GetMemoryValidationLevel(), ReturnAddress);
+	}
+
+	Work(Arg);
+
+	FContext* Context = FContext::Get();
+	if (Transaction == Context->GetCurrentTransaction() && Transaction->IsOpenActive())  // Transaction may have been aborted.
+	{
+		Transaction->SetClosedActive();
+	}
+
+	if (Context->IsAborting())
+	{
+		Context->Throw();
+	}
+}
+
 AUTORTFM_CORE_API void RTFM_autortfm_record_open_write_err(void*, size_t)
 {
 	AUTORTFM_FATAL("The function `autortfm_record_open_write` was called from closed code");
@@ -875,7 +897,7 @@ extern "C" void autortfm_open_explicit_validation(autortfm_memory_validation_lev
 {
 	if (autortfm_is_closed())
 	{
-		return RTFM_autortfm_open(ValidationLevel, Work, Arg, ReturnAddress);
+		return autortfm_open_from_closed_explicit_validation(ValidationLevel, Work, Arg, ReturnAddress);
 	}
 
 	Work(Arg);
@@ -885,7 +907,7 @@ extern "C" void autortfm_open(void (*Work)(void*), void* Arg, const void* Return
 {
 	if (autortfm_is_closed())
 	{
-		return RTFM_autortfm_open(autortfm_memory_validation_level_default, Work, Arg, ReturnAddress);
+		return autortfm_open_from_closed(Work, Arg, ReturnAddress);
 	}
 
 	Work(Arg);
