[AutoRTFM] Improve eliding of unnecessary closed variant functions

This includes lambdas / private-linkage functions that are only ever called from AutoRTFM::Open() / autortfm_open().

Allows us to enable compile-time instruction validation of closed functions without considering autortfm_open'd functions as closed.
Also reduces server DebugGame executable size by ~40KiB.

diff --git a/Engine/Source/Runtime/AutoRTFM/Public/AutoRTFM.h b/Engine/Source/Runtime/AutoRTFM/Public/AutoRTFM.h
--- a/Engine/Source/Runtime/AutoRTFM/Public/AutoRTFM.h
+++ b/Engine/Source/Runtime/AutoRTFM/Public/AutoRTFM.h
@@ -893,7 +893,7 @@ UE_AUTORTFM_CRITICAL_INLINE ReturnType Open(const FunctorType& Functor)
 
 			if constexpr (VALIDATION_LEVEL == EMemoryValidationLevel::Default)
 			{
-				autortfm_open(&FCallHelper::Call, const_cast<void*>(static_cast<const void*>(&Functor)), __builtin_return_address(0));
+				autortfm_open(&FCallHelper::Call, const_cast<void*>(reinterpret_cast<const void*>(&Functor)), __builtin_return_address(0));
 			}
 			else
 			{
