The underlying TStack should use the same validation as the TBitStack.

diff --git a/Engine/Source/Runtime/AutoRTFM/Private/BitStack.h b/Engine/Source/Runtime/AutoRTFM/Private/BitStack.h
--- a/Engine/Source/Runtime/AutoRTFM/Private/BitStack.h
+++ b/Engine/Source/Runtime/AutoRTFM/Private/BitStack.h
@@ -184,7 +184,7 @@ private:
 		return (NumBits + (NumWordBits-1)) / NumWordBits;
 	}
 
-	TStack<WordType, NumWordsFor(InlineCapacity)> Words;
+	TStack<WordType, NumWordsFor(InlineCapacity), Validation> Words;
 	size_t Count = 0; // In bits
 };
 
