[AutoRTFM] Implement a custom replacement for setjmp() / longjmp() for Windows + X64

The standard windows setjmp / longjmp implementation internally uses exceptions to jump back to the setjmp() call.
This results in destructors being called when unwinding the stack, which is problematic for AutoRTFM as the unwind occurs after AutoRTFM tracked memory is reverted. This can cause crashes and memory corruption.

This CL implements a replacement for Windows + X64. This implementation saves the register and restores the register state without performing stack unwinding.
Benchmarking suggests this ~28x faster to throw and ~1.12x faster to enter the try.

diff --git a/Engine/Source/Programs/AutoRTFMTests/Private/LongJumpTests.cpp b/Engine/Source/Programs/AutoRTFMTests/Private/LongJumpTests.cpp
new file mode 100644
--- /dev/null
+++ b/Engine/Source/Programs/AutoRTFMTests/Private/LongJumpTests.cpp
@@ -0,0 +1,144 @@
+// Copyright Epic Games, Inc. All Rights Reserved.
+
+#include "LongJump.h"
+#include "Catch2Includes.h"
+
+TEST_CASE("LongJump")
+{
+	struct FSetOnDestruct
+	{
+		bool& bDestructed;
+		FSetOnDestruct(bool& bDestructed) : bDestructed{bDestructed} {}
+		~FSetOnDestruct() { bDestructed = true; }
+	};
+	SECTION("NoThrow")
+	{
+		bool bTryCalled = false;
+		bool bCatchCalled = false;
+
+		AutoRTFM::FLongJump LongJump;
+		LongJump.TryCatch(
+		/* Try */ [&]
+		{
+			bTryCalled = true;
+		},
+		/* Catch */ [&]
+		{
+			bCatchCalled = true;
+		});
+		
+		REQUIRE(bTryCalled);
+		REQUIRE(!bCatchCalled);
+	}
+	SECTION("Throw")
+	{
+		bool bPostThrow = false;
+		bool bCatchCalled = false;
+
+		AutoRTFM::FLongJump LongJump;
+		LongJump.TryCatch(
+		/* Try */ [&]
+		{
+			LongJump.Throw();
+			bPostThrow = true;
+		},
+		/* Catch */ [&]
+		{
+			bCatchCalled = true;
+		});
+		
+		REQUIRE(!bPostThrow);
+		REQUIRE(bCatchCalled);
+	}
+	SECTION("NoUnwind")
+	{
+		SECTION("Direct")
+		{
+			bool bDtorCalled = false;
+
+			AutoRTFM::FLongJump LongJump;
+			LongJump.TryCatch(
+			/* Try */ [&]
+			{
+				FSetOnDestruct SetOnDestruct{bDtorCalled};
+				LongJump.Throw();
+			},
+			/* Catch */ [&]
+			{});
+			
+			REQUIRE(!bDtorCalled);
+		}
+		SECTION("Nested")
+		{
+			bool bDtorCalled = false;
+
+			AutoRTFM::FLongJump LongJump;
+			LongJump.TryCatch(
+			/* Try */ [&]
+			{
+				struct FNestedCalls
+				{
+					UE_AUTORTFM_FORCENOINLINE void F3(AutoRTFM::FLongJump& LongJump)
+					{
+						LongJump.Throw();
+					}
+					UE_AUTORTFM_FORCENOINLINE void F2(AutoRTFM::FLongJump& LongJump)
+					{
+						F3(LongJump);
+					}
+					UE_AUTORTFM_FORCENOINLINE void F1(AutoRTFM::FLongJump& LongJump)
+					{
+						F2(LongJump);
+					}
+				};
+				FSetOnDestruct SetOnDestruct{bDtorCalled};
+				FNestedCalls{}.F1(LongJump);
+			},
+			/* Catch */ [&]
+			{
+			});
+			
+			REQUIRE(!bDtorCalled);
+		}
+	}
+}
+
+TEST_CASE("LongJump.Benchmarks")
+{
+	BENCHMARK("NoThrow", InputValue)
+	{
+		int DoNotOptimizeAway = 0;
+		for (int I = 0; I < 1024; I++)
+		{
+			AutoRTFM::FLongJump LongJump;
+			LongJump.TryCatch(
+			/* Try */ [&]
+			{
+				DoNotOptimizeAway += InputValue;
+			},
+			/* Catch */ [&]
+			{
+			});
+		}
+		return DoNotOptimizeAway;
+	};
+
+	BENCHMARK("Throw", InputValue)
+	{
+		int DoNotOptimizeAway = 0;
+		for (int I = 0; I < 1024; I++)
+		{
+			AutoRTFM::FLongJump LongJump;
+			LongJump.TryCatch(
+			/* Try */ [&]
+			{
+				LongJump.Throw();
+			},
+			/* Catch */ [&]
+			{
+				DoNotOptimizeAway += InputValue;
+			});
+		}
+		return DoNotOptimizeAway;
+	};
+}
diff --git a/Engine/Source/Runtime/AutoRTFM/Private/BuildMacros.h b/Engine/Source/Runtime/AutoRTFM/Private/BuildMacros.h
--- a/Engine/Source/Runtime/AutoRTFM/Private/BuildMacros.h
+++ b/Engine/Source/Runtime/AutoRTFM/Private/BuildMacros.h
@@ -59,3 +59,18 @@ static_assert(AUTORTFM_BUILD_DEBUG + AUTORTFM_BUILD_DEVELOPMENT + AUTORTFM_BUILD
 		#define AUTORTFM_PLATFORM_LINUX 0
 	#endif
 #endif
+
+#if defined(__x86_64__) || defined(_M_X64)
+	#define AUTORTFM_ARCHITECTURE_X64 1
+#else
+	#define AUTORTFM_ARCHITECTURE_X64 0
+#endif
+
+#if defined(__has_feature)
+	#if __has_feature(cxx_exceptions)
+		#define AUTORTFM_EXCEPTIONS_ENABLED 1
+	#endif
+#endif
+#if !defined(AUTORTFM_EXCEPTIONS_ENABLED)
+	#define AUTORTFM_EXCEPTIONS_ENABLED 0
+#endif
diff --git a/Engine/Source/Runtime/AutoRTFM/Private/LongJump.cpp b/Engine/Source/Runtime/AutoRTFM/Private/LongJump.cpp
new file mode 100644
--- /dev/null
+++ b/Engine/Source/Runtime/AutoRTFM/Private/LongJump.cpp
@@ -0,0 +1,103 @@
+// Copyright Epic Games, Inc. All Rights Reserved.
+
+#if (defined(__AUTORTFM) && __AUTORTFM)
+
+#include "LongJump.h"
+
+#if AUTORTFM_USE_SAVE_RESTORE_CONTEXT
+namespace AutoRTFM
+{
+
+// https://learn.microsoft.com/en-us/cpp/build/x64-calling-convention?view=msvc-170#callercallee-saved-registers
+// Registers that must be saved:
+// RBX, RBP, RDI, RSI, RSP, R12, R13, R14, R15, and XMM6-XMM15
+UE_AUTORTFM_NOAUTORTFM __attribute__((noinline)) __attribute((naked))
+bool FLongJump::SaveContext(void* Buffer)
+{
+	__asm
+	{
+		// rcx: Buffer, rax: Return value
+
+		// Save non-volatile registers
+		mov     [rcx +   0], rsp
+		mov     [rcx +   8], rbx
+		mov     [rcx +  16], rbp
+		mov     [rcx +  24], rdi
+		mov     [rcx +  32], rsi
+		mov     [rcx +  40], r12
+		mov     [rcx +  48], r13
+		mov     [rcx +  56], r14
+		mov     [rcx +  64], r15
+		movupd  [rcx +  80], xmm6
+		movupd  [rcx +  96], xmm7
+		movupd  [rcx + 112], xmm8
+		movupd  [rcx + 128], xmm9
+		movupd  [rcx + 144], xmm10
+		movupd  [rcx + 160], xmm11
+		movupd  [rcx + 176], xmm12
+		movupd  [rcx + 192], xmm13
+		movupd  [rcx + 208], xmm14
+		movupd  [rcx + 224], xmm15
+		
+		// Save register control & status state
+		stmxcsr [rcx + 240]
+		fnstcw  [rcx + 244]
+
+		// Use return address on stack as restore instruction pointer
+		mov     r8,          [rsp]
+		mov     [rcx + 248], r8
+
+		// Return 0 to indicate that this is a context-save
+		mov     rax, 0
+		ret
+	}
+}
+UE_AUTORTFM_NOAUTORTFM __attribute__((noinline)) __attribute((naked))
+void FLongJump::RestoreContext(void* Buffer)
+{
+	__asm
+	{
+		// rcx: Buffer, rax: Return value
+
+		// Restore non-volatile registers
+		mov     rsp,   [rcx +   0]
+		mov     rbx,   [rcx +   8]
+		mov     rbp,   [rcx +  16]
+		mov     rdi,   [rcx +  24]
+		mov     rsi,   [rcx +  32]
+		mov     r12,   [rcx +  40]
+		mov     r13,   [rcx +  48]
+		mov     r14,   [rcx +  56]
+		mov     r15,   [rcx +  64]
+		movupd  xmm6,  [rcx +  80]
+		movupd  xmm7,  [rcx +  96]
+		movupd  xmm8,  [rcx + 112]
+		movupd  xmm9,  [rcx + 128]
+		movupd  xmm10, [rcx + 144]
+		movupd  xmm11, [rcx + 160]
+		movupd  xmm12, [rcx + 176]
+		movupd  xmm13, [rcx + 192]
+		movupd  xmm14, [rcx + 208]
+		movupd  xmm15, [rcx + 224]
+
+		// Restore register control & status state
+		ldmxcsr        [rcx + 240]
+		fldcw          [rcx + 244]
+		fnclex
+
+		// Replace caller's return address with saved return address and return.
+		mov     r8,    [rcx + 248]
+		mov     [rsp], r8
+
+		// Return 1 to indicate that this is a context-restore
+		mov     rax, 1
+		ret
+	}
+}
+
+}
+
+#endif // AUTORTFM_USE_SAVE_RESTORE_CONTEXT
+
+#endif // (defined(__AUTORTFM) && __AUTORTFM)
+
diff --git a/Engine/Source/Runtime/AutoRTFM/Private/LongJump.h b/Engine/Source/Runtime/AutoRTFM/Private/LongJump.h
--- a/Engine/Source/Runtime/AutoRTFM/Private/LongJump.h
+++ b/Engine/Source/Runtime/AutoRTFM/Private/LongJump.h
@@ -9,6 +9,12 @@
 
 #include <setjmp.h>
 
+#if AUTORTFM_PLATFORM_WINDOWS && AUTORTFM_ARCHITECTURE_X64 // We only have an Windows X64 implementation right now.
+#define AUTORTFM_USE_SAVE_RESTORE_CONTEXT 1
+#else
+#define AUTORTFM_USE_SAVE_RESTORE_CONTEXT 0
+#endif
+
 namespace AutoRTFM
 {
 
@@ -18,51 +24,63 @@ namespace AutoRTFM
 class FLongJump
 {
 public:
-    FLongJump()
-    {
-        memset(this, 0, sizeof(*this));
-    }
+	FLongJump()
+	{
+		memset(this, 0, sizeof(*this));
+	}
 
-    template<typename TTryFunctor, typename TCatchFunctor>
-    void TryCatch(const TTryFunctor& TryFunctor, const TCatchFunctor& CatchFunctor);
-    
-    [[noreturn]] void Throw();
-    
+	template<typename TTryFunctor, typename TCatchFunctor>
+	void TryCatch(const TTryFunctor& TryFunctor, const TCatchFunctor& CatchFunctor);
+	
+	[[noreturn]] void Throw();
+	
 private:
-    jmp_buf JmpBuf;
-    bool bIsSet;
+#if AUTORTFM_USE_SAVE_RESTORE_CONTEXT
+	static bool SaveContext(void* Buffer);
+	[[noreturn]] static void RestoreContext(void* Context);
+	static constexpr size_t ContextSize = 256;
+	static constexpr size_t ContextAlignment = 16;
+	alignas(ContextAlignment) std::byte Context[ContextSize];
+#else
+	jmp_buf JmpBuf;
+#endif
+	bool bIsSet;
 };
 
 template<typename TTryFunctor, typename TCatchFunctor>
 void FLongJump::TryCatch(const TTryFunctor& TryFunctor, const TCatchFunctor& CatchFunctor)
 {
-    AUTORTFM_ASSERT(!bIsSet);
-#if AUTORTFM_PLATFORM_WINDOWS
+	AUTORTFM_ASSERT(!bIsSet);
+#if AUTORTFM_USE_SAVE_RESTORE_CONTEXT
+	if (!SaveContext(Context))
+#elif AUTORTFM_PLATFORM_WINDOWS
 	if (!setjmp(JmpBuf))
 #else
 	if (!_setjmp(JmpBuf))
-#endif // AUTORTFM_PLATFORM_WINDOWS
-    {
-        bIsSet = true;
-        TryFunctor();
-        bIsSet = false;
-    }
-    else
-    {
-        AUTORTFM_ASSERT(bIsSet);
-        bIsSet = false;
-        CatchFunctor();
-    }
+#endif
+	{
+		bIsSet = true;
+		TryFunctor();
+		bIsSet = false;
+	}
+	else
+	{
+		AUTORTFM_ASSERT(bIsSet);
+		bIsSet = false;
+		CatchFunctor();
+	}
 }
 
 inline void FLongJump::Throw()
 {
-    AUTORTFM_ASSERT(bIsSet);
-#if AUTORTFM_PLATFORM_WINDOWS
+	AUTORTFM_ASSERT(bIsSet);
+#if AUTORTFM_USE_SAVE_RESTORE_CONTEXT
+	RestoreContext(Context);
+#elif AUTORTFM_PLATFORM_WINDOWS
 	longjmp(JmpBuf, 1);
 #else
 	_longjmp(JmpBuf, 1);
-#endif // AUTORTFM_PLATFORM_WINDOWS
+#endif
 }
 
 } // namespace AutoRTFM
diff --git a/Engine/Source/Runtime/AutoRTFM/Private/WriteLog.cpp b/Engine/Source/Runtime/AutoRTFM/Private/WriteLog.cpp
--- a/Engine/Source/Runtime/AutoRTFM/Private/WriteLog.cpp
+++ b/Engine/Source/Runtime/AutoRTFM/Private/WriteLog.cpp
@@ -10,7 +10,7 @@
 namespace AutoRTFM
 {
 
-#if defined(__x86_64__) || defined(_M_X64)
+#if AUTORTFM_ARCHITECTURE_X64
 #define WRITELOG_HASH_MAYBE_HAS_AVX2 1
 #include <immintrin.h>
 #else
