[AutoRTFM] Support exceptions with AutoRTFM instrumentation

This CL adds support for compiling with both the `-fexceptions` and `-fautortfm` flags.

verse-clang changes:
* General handling for LLVM's `invoke` instruction. This involved considering all uses of `call` instructions and determining whether the logic also needed to also support `invoke`.
* Correct cloning and hoisting of exception unwinding blocks for static initializers.
* Propagation of funclet operand bundles when performing closed function lookups in unwind blocks.
* A bunch of cleanup around imported functions.
* Improved diagnostic messages and behavior for when Things Go Wrong.
* Fixup of existing lit tests for relaxed rules, minor codegen diffs, and new exception tests.

Runtime changes:
* AutoRTFM public API has now been decorated with `AUTORTFM_NOEXCEPT` / `AUTORTFM_EXCEPT` annotations.
* `FContext::Transact()` now automatically commits the transaction if an exception is thrown inside a transaction, which is caught outside the transaction.
* A number of new open -> closed standard library functions are now registered to handle throwing and catching exceptions.
* The number of errors when #including "AutoRTFM.h" from a C file have been reduced, however this still doesn't compile as we use `bool` in the public "C" API, which is not a C type.

AutoRTFMTests changes:
* A new exceptions-enabled version of these tests has been added as `AutoRTFMTestsWithExceptions`.
* `AutoRTFMTestsWithExceptionsTarget` has been added to `PreSubmitTest.py`
* Added a number of new exception unit tests.

diff --git a/Engine/Source/Programs/AutoRTFMTests/AutoRTFMTestsWithExceptions.Target.cs b/Engine/Source/Programs/AutoRTFMTests/AutoRTFMTestsWithExceptions.Target.cs
new file mode 100644
--- /dev/null
+++ b/Engine/Source/Programs/AutoRTFMTests/AutoRTFMTestsWithExceptions.Target.cs
@@ -0,0 +1,15 @@
+// Copyright Epic Games, Inc. All Rights Reserved.
+
+using EpicGames.Core;
+using UnrealBuildTool;
+using System.Collections.Generic;
+using UnrealBuildBase;
+
+[SupportedPlatforms(UnrealPlatformClass.All)]
+public class AutoRTFMTestsWithExceptionsTarget : AutoRTFMTestsTarget
+{
+	public AutoRTFMTestsWithExceptionsTarget(TargetInfo Target) : base(Target)
+	{
+		bForceEnableExceptions = true;
+	}
+}
diff --git a/Engine/Source/Programs/AutoRTFMTests/Private/ExceptionTests.cpp b/Engine/Source/Programs/AutoRTFMTests/Private/ExceptionTests.cpp
new file mode 100644
--- /dev/null
+++ b/Engine/Source/Programs/AutoRTFMTests/Private/ExceptionTests.cpp
@@ -0,0 +1,223 @@
+// Copyright Epic Games, Inc. All Rights Reserved.
+
+#include "AutoRTFMTesting.h"
+#include "AutoRTFMTestUtils.h"
+#include "Catch2Includes.h"
+#include "BuildMacros.h"
+
+#include <string>
+#include <vector>
+
+#if AUTORTFM_EXCEPTIONS_ENABLED
+
+TEST_CASE("Exceptions")
+{
+	// Tests are sensitive to retries. Disable for these tests.
+	AutoRTFMTestUtils::FScopedRetry Retry(AutoRTFM::ForTheRuntime::EAutoRTFMRetryTransactionState::NoRetry);
+
+	using EventList = std::vector<std::string>;
+	EventList Events;
+
+	auto Event = [&Events](std::string_view Event)
+	{
+		UE_AUTORTFM_OPEN
+		{
+			Events.push_back(std::string(Event));
+		};
+	};
+
+	struct FScope
+	{
+		UE_AUTORTFM_ALWAYS_OPEN
+		FScope(EventList& Events, const char* Name) : Events{Events}, Name{Name}
+		{
+			Events.push_back(std::string{Name});
+		}
+		UE_AUTORTFM_ALWAYS_OPEN
+		~FScope()
+		{
+			Events.push_back(std::string{"End-"} + Name);
+		}
+		EventList& Events;
+		char const* const Name;
+	};
+#define SCOPE(INDEX) FScope UE_AUTORTFM_CONCAT(Scope_, __COUNTER__){Events, INDEX}
+
+	SECTION("Try(Transact())")
+	{
+		{
+			SCOPE("Outer");
+			try
+			{
+				SCOPE("Try");
+				AutoRTFM::Testing::Commit([&]
+				{
+					SCOPE("Transact");
+				});
+			}
+			catch (...)
+			{
+				SCOPE("Catch");
+			}
+		}
+
+		const EventList Expected
+		{
+			"Outer",
+			"Try",
+			"Transact",
+			"End-Transact",
+			"End-Try",
+			"End-Outer",
+		};
+		REQUIRE(Events == Expected);
+	}
+
+	SECTION("Try(Transact(Abort))")
+	{
+		{
+			SCOPE("Outer");
+			try
+			{
+				SCOPE("Try");
+				AutoRTFM::Testing::Abort([&]
+				{
+					SCOPE("Transact");
+					Event("Abort");
+					AutoRTFM::AbortTransaction();
+				});
+			}
+			catch (...)
+			{
+				SCOPE("Catch");
+			}
+		}
+
+		const EventList Expected
+		{
+			"Outer",
+			"Try",
+			"Transact",
+			"Abort",
+			"End-Try",
+			"End-Outer",
+		};
+		REQUIRE(Events == Expected);
+	}
+
+	SECTION("Try(Transact(Throw))")
+	{
+		{
+			SCOPE("Outer");
+			try
+			{
+				SCOPE("Try");
+				AutoRTFM::Testing::Commit([&]
+				{
+					SCOPE("Transact");
+					Event("Throw");
+					throw(42);
+				});
+			}
+			catch (int I)
+			{
+				SCOPE("Catch");
+				REQUIRE(I == 42);
+			}
+		}
+
+		const EventList Expected
+		{
+			"Outer",
+			"Try",
+			"Transact",
+			"Throw",
+			"End-Transact",
+			"End-Try",
+			"Catch",
+			"End-Catch",
+			"End-Outer",
+		};
+		REQUIRE(Events == Expected);
+	}
+
+	SECTION("Transact(Try(Throw))")
+	{
+		{
+			SCOPE("Outer");
+			AutoRTFM::Testing::Commit([&]
+			{
+				SCOPE("Transact");
+				try
+				{
+					SCOPE("Try");
+					Event("Throw");
+					throw(42);
+				}
+				catch (int I)
+				{
+					SCOPE("Catch");
+					REQUIRE(I == 42);
+				}
+			});
+		}
+
+		const EventList Expected
+		{
+			"Outer",
+			"Transact",
+			"Try",
+			"Throw",
+			"End-Try",
+			"Catch",
+			"End-Catch",
+			"End-Transact",
+			"End-Outer",
+		};
+		REQUIRE(Events == Expected);
+	}
+
+	SECTION("Transact(Try(Transact(Throw)))")
+	{
+		{
+			SCOPE("Outer");
+			AutoRTFM::Testing::Commit([&]
+			{
+				SCOPE("Transact-A");
+				try
+				{
+					SCOPE("Try");
+					AutoRTFM::Testing::Commit([&]
+					{
+						SCOPE("Transact-B");
+						Event("Throw");
+						throw(42);
+					});
+				}
+				catch (int I)
+				{
+					SCOPE("Catch");
+					REQUIRE(I == 42);
+				}
+			});
+		}
+
+		const EventList Expected
+		{
+			"Outer",
+			"Transact-A",
+			"Try",
+			"Transact-B",
+			"Throw",
+			"End-Transact-B",
+			"End-Try",
+			"Catch",
+			"End-Catch",
+			"End-Transact-A",
+			"End-Outer",
+		};
+		REQUIRE(Events == Expected);
+	}
+}
+
+#endif // AUTORTFM_EXCEPTIONS_ENABLED
diff --git a/Engine/Source/Runtime/AutoRTFM/Private/API.cpp b/Engine/Source/Runtime/AutoRTFM/Private/API.cpp
--- a/Engine/Source/Runtime/AutoRTFM/Private/API.cpp
+++ b/Engine/Source/Runtime/AutoRTFM/Private/API.cpp
@@ -722,7 +722,7 @@ AUTORTFM_CORE_API void ProcessAllPendingOpenToClosedRegistrations()
 
 // The AutoRTFM public API.
 
-extern "C" void autortfm_initialize(const autortfm_extern_api* ExternAPI)
+extern "C" void autortfm_initialize(const autortfm_extern_api* ExternAPI) noexcept
 {
 	AUTORTFM_ENSURE_MSG(!FContext::Get(), "AutoRTFM initialized twice");
 	AUTORTFM_ASSERT(ExternAPI);
@@ -737,7 +737,7 @@ extern "C" void autortfm_initialize(const autortfm_extern_api* ExternAPI)
 
 // Each function will be forked by the compiler into an open and closed variant.
 // autortfm_is_closed() is used to branch to the closed variants declared above.
-extern "C" bool autortfm_is_transactional()
+extern "C" bool autortfm_is_transactional() noexcept
 {
 	if (autortfm_is_closed())
 	{
@@ -753,7 +753,7 @@ extern "C" bool autortfm_is_transactional()
 	return false;
 }
 
-extern "C" bool autortfm_is_committing_or_aborting()
+extern "C" bool autortfm_is_committing_or_aborting() noexcept
 {
 	if (ForTheRuntime::IsAutoRTFMRuntimeEnabled())
 	{
@@ -801,7 +801,7 @@ extern "C" void autortfm_commit(void (*UninstrumentedWork)(void*), void (*Instru
 	AUTORTFM_FATAL_IF(Result != autortfm_committed, "Unexpected transaction result: %u", Result);
 }
 
-extern "C" void autortfm_abort()
+extern "C" void autortfm_abort() noexcept
 {
 	if (autortfm_is_closed())
 	{
@@ -812,7 +812,7 @@ extern "C" void autortfm_abort()
 	FContext::Get()->AbortByRequestAndThrow();
 }
 
-extern "C" bool autortfm_start_transaction()
+extern "C" bool autortfm_start_transaction() noexcept
 {
 	if (autortfm_is_closed())
 	{
@@ -826,7 +826,7 @@ extern "C" bool autortfm_start_transaction()
 	return true;
 }
 
-extern "C" autortfm_result autortfm_commit_transaction()
+extern "C" autortfm_result autortfm_commit_transaction() noexcept
 {
 	if (autortfm_is_closed())
 	{
@@ -838,13 +838,13 @@ extern "C" autortfm_result autortfm_commit_transaction()
 	return static_cast<autortfm_result>(FContext::Get()->CommitTransaction());
 }
 
-extern "C" void autortfm_abort_transaction()
+extern "C" void autortfm_abort_transaction() noexcept
 {
 	AUTORTFM_ASSERT(autortfm_is_closed()); // RollbackTransaction should be used when in the open.
 	RTFM_autortfm_abort_transaction();
 }
 
-extern "C" autortfm_result autortfm_rollback_transaction()
+extern "C" autortfm_result autortfm_rollback_transaction() noexcept
 {
 	AUTORTFM_ASSERT(!autortfm_is_closed()); // AbortTransaction should be used when in the closed.
 	FContext* const Context = FContext::Get();
@@ -852,7 +852,7 @@ extern "C" autortfm_result autortfm_rollback_transaction()
 	return static_cast<autortfm_result>(Context->RollbackTransaction(EContextStatus::AbortedByRequest));
 }
 
-extern "C" autortfm_result autortfm_cascading_rollback_transaction()
+extern "C" autortfm_result autortfm_cascading_rollback_transaction() noexcept
 {
 	AUTORTFM_ASSERT(!autortfm_is_closed());
 
@@ -861,7 +861,7 @@ extern "C" autortfm_result autortfm_cascading_rollback_transaction()
 	return static_cast<autortfm_result>(Context->RollbackTransaction(EContextStatus::AbortedByCascadingAbort));
 }
 
-extern "C" void autortfm_clear_transaction_status()
+extern "C" void autortfm_clear_transaction_status() noexcept
 {
 	if (autortfm_is_closed())
 	{
@@ -872,12 +872,12 @@ extern "C" void autortfm_clear_transaction_status()
 	FContext::Get()->ClearTransactionStatus();
 }
 
-extern "C" autortfm_status autortfm_get_context_status()
+extern "C" autortfm_status autortfm_get_context_status() noexcept
 {
 	return static_cast<autortfm_status>(FContext::Get()->GetStatus());
 }
 
-extern "C" UE_AUTORTFM_NOAUTORTFM bool autortfm_is_aborting()
+extern "C" UE_AUTORTFM_NOAUTORTFM bool autortfm_is_aborting() noexcept
 {
 	if (ForTheRuntime::IsAutoRTFMRuntimeEnabled())
 	{
@@ -887,7 +887,7 @@ extern "C" UE_AUTORTFM_NOAUTORTFM bool autortfm_is_aborting()
 	return false;
 }
 
-extern "C" UE_AUTORTFM_NOAUTORTFM bool autortfm_current_nest_throw()
+extern "C" UE_AUTORTFM_NOAUTORTFM bool autortfm_current_nest_throw() noexcept
 {
 	FContext::Get()->Throw();
 	return true;
@@ -954,7 +954,7 @@ extern "C" autortfm_status autortfm_close(void (*UninstrumentedWork)(void*), voi
 	return Result;
 }
 
-extern "C" void autortfm_record_open_write(void* Ptr, size_t Size)
+extern "C" void autortfm_record_open_write(void* Ptr, size_t Size) noexcept
 {
 	if (autortfm_is_closed())
 	{
@@ -969,7 +969,7 @@ extern "C" void autortfm_record_open_write(void* Ptr, size_t Size)
 	}
 }
 
-extern "C" void autortfm_record_open_write_no_memory_validation(void* Ptr, size_t Size)
+extern "C" void autortfm_record_open_write_no_memory_validation(void* Ptr, size_t Size) noexcept
 {
 	if (autortfm_is_closed())
 	{
@@ -984,7 +984,7 @@ extern "C" void autortfm_record_open_write_no_memory_validation(void* Ptr, size_
 	}
 }
 
-extern "C" UE_AUTORTFM_NOAUTORTFM void autortfm_register_open_to_closed_functions(autortfm_open_to_closed_table* Table)
+extern "C" UE_AUTORTFM_NOAUTORTFM void autortfm_register_open_to_closed_functions(autortfm_open_to_closed_table* Table) noexcept
 {
 	Table->Next = GPendingOpenToClosedRegistrations;
 	GPendingOpenToClosedRegistrations = Table;
@@ -995,7 +995,7 @@ extern "C" UE_AUTORTFM_NOAUTORTFM void autortfm_register_open_to_closed_function
 	}
 }
 
-extern "C" UE_AUTORTFM_NOAUTORTFM void autortfm_unregister_open_to_closed_functions(autortfm_open_to_closed_table* Table)
+extern "C" UE_AUTORTFM_NOAUTORTFM void autortfm_unregister_open_to_closed_functions(autortfm_open_to_closed_table* Table) noexcept
 {
 	if (Table == GPendingOpenToClosedRegistrations)
 	{
@@ -1024,7 +1024,7 @@ extern "C" UE_AUTORTFM_NOAUTORTFM void autortfm_unregister_open_to_closed_functi
 	// then something will need to be done here.
 }
 
-extern "C" bool autortfm_is_on_current_transaction_stack(void* Ptr)
+extern "C" bool autortfm_is_on_current_transaction_stack(void* Ptr) noexcept
 {
 	if (FContext* Context = FContext::Get())
 	{
@@ -1135,7 +1135,7 @@ extern "C" void autortfm_on_commit(void (*Work)(void*), void* Arg)
     Work(Arg);
 }
 
-extern "C" void autortfm_on_abort(void (*Work)(void*), void* Arg)
+extern "C" void autortfm_on_abort(void (*Work)(void*), void* Arg) noexcept
 {
 	if (autortfm_is_closed())
 	{
@@ -1144,7 +1144,7 @@ extern "C" void autortfm_on_abort(void (*Work)(void*), void* Arg)
 
 }
 
-extern "C" void autortfm_push_on_abort_handler(const void* Key, void (*Work)(void*), void* Arg)
+extern "C" void autortfm_push_on_abort_handler(const void* Key, void (*Work)(void*), void* Arg) noexcept
 {
 	if (autortfm_is_closed())
 	{
@@ -1153,7 +1153,7 @@ extern "C" void autortfm_push_on_abort_handler(const void* Key, void (*Work)(voi
 
 }
 
-extern "C" void autortfm_pop_on_abort_handler(const void* Key)
+extern "C" void autortfm_pop_on_abort_handler(const void* Key) noexcept
 {
 	if (autortfm_is_closed())
 	{
@@ -1162,7 +1162,7 @@ extern "C" void autortfm_pop_on_abort_handler(const void* Key)
 
 }
 
-extern "C" void* autortfm_did_allocate(void* Ptr, size_t Size)
+extern "C" void* autortfm_did_allocate(void* Ptr, size_t Size) noexcept
 {
 	if (autortfm_is_closed())
 	{
@@ -1172,7 +1172,7 @@ extern "C" void* autortfm_did_allocate(void* Ptr, size_t Size)
     return Ptr;
 }
 
-extern "C" void autortfm_did_free(void* Ptr)
+extern "C" void autortfm_did_free(void* Ptr) noexcept
 {
 	if (autortfm_is_closed())
 	{
@@ -1200,7 +1200,7 @@ extern "C" void autortfm_did_free(void* Ptr)
 // compiler pass. Should not be called manually by the user, a call to this will
 // be injected by the compiler into a global constructor in the AutoRTFM compiled
 // code.
-extern "C" UE_AUTORTFM_NOAUTORTFM void autortfm_check_abi(void* const Ptr, const size_t Size)
+extern "C" UE_AUTORTFM_NOAUTORTFM void autortfm_check_abi(void* const Ptr, const size_t Size) noexcept
 {
     struct FConstants final
     {
diff --git a/Engine/Source/Runtime/AutoRTFM/Private/CompilerInterface.cpp b/Engine/Source/Runtime/AutoRTFM/Private/CompilerInterface.cpp
--- a/Engine/Source/Runtime/AutoRTFM/Private/CompilerInterface.cpp
+++ b/Engine/Source/Runtime/AutoRTFM/Private/CompilerInterface.cpp
@@ -54,42 +54,42 @@ void CheckAlignment(FContext* Context, void* Ptr, size_t AlignmentMask)
 
 #endif
 
-extern "C" UE_AUTORTFM_API AUTORTFM_NO_ASAN void autortfm_record_write(void* Ptr, size_t Size)
+extern "C" UE_AUTORTFM_API AUTORTFM_NO_ASAN void autortfm_record_write(void* Ptr, size_t Size) noexcept
 {
 	UE_AUTORTFM_HANDLE_NULL_WRITE(Ptr);
 	FContext* Context = FContext::Get();
 	Context->RecordWrite(Ptr, Size);
 }
 
-extern "C" UE_AUTORTFM_API AUTORTFM_NO_ASAN void autortfm_record_write_1(void* Ptr)
+extern "C" UE_AUTORTFM_API AUTORTFM_NO_ASAN void autortfm_record_write_1(void* Ptr) noexcept
 {
 	UE_AUTORTFM_HANDLE_NULL_WRITE(Ptr);
 	FContext* Context = FContext::Get();
 	Context->RecordWrite<1>(Ptr);
 }
 
-extern "C" UE_AUTORTFM_API AUTORTFM_NO_ASAN void autortfm_record_write_2(void* Ptr)
+extern "C" UE_AUTORTFM_API AUTORTFM_NO_ASAN void autortfm_record_write_2(void* Ptr) noexcept
 {
 	UE_AUTORTFM_HANDLE_NULL_WRITE(Ptr);
 	FContext* Context = FContext::Get();
 	Context->RecordWrite<2>(Ptr);
 }
 
-extern "C" UE_AUTORTFM_API AUTORTFM_NO_ASAN void autortfm_record_write_4(void* Ptr)
+extern "C" UE_AUTORTFM_API AUTORTFM_NO_ASAN void autortfm_record_write_4(void* Ptr) noexcept
 {
 	UE_AUTORTFM_HANDLE_NULL_WRITE(Ptr);
 	FContext* Context = FContext::Get();
 	Context->RecordWrite<4>(Ptr);
 }
 
-extern "C" UE_AUTORTFM_API AUTORTFM_NO_ASAN void autortfm_record_write_8(void* Ptr)
+extern "C" UE_AUTORTFM_API AUTORTFM_NO_ASAN void autortfm_record_write_8(void* Ptr) noexcept
 {
 	UE_AUTORTFM_HANDLE_NULL_WRITE(Ptr);
 	FContext* Context = FContext::Get();
 	Context->RecordWrite<8>(Ptr);
 }
 
-extern "C" UE_AUTORTFM_API AUTORTFM_NO_ASAN void autortfm_record_masked_write(void* Ptr, uintptr_t Mask, int MaskWidthBits, int ValueSizeBytes)
+extern "C" UE_AUTORTFM_API AUTORTFM_NO_ASAN void autortfm_record_masked_write(void* Ptr, uintptr_t Mask, int MaskWidthBits, int ValueSizeBytes) noexcept
 {
 	UE_AUTORTFM_HANDLE_NULL_WRITE(Ptr);
 
@@ -113,35 +113,35 @@ extern "C" UE_AUTORTFM_API AUTORTFM_NO_ASAN void autortfm_record_masked_write(vo
 UE_AUTORTFM_REGISTER_OPEN_TO_CLOSED_FUNCTIONS(
 	UE_AUTORTFM_MAP_OPEN_TO_SELF(autortfm_lookup_function)
 );
-extern "C" UE_AUTORTFM_NOAUTORTFM UE_AUTORTFM_API void* autortfm_lookup_function(void* OriginalFunction, const char* Where)
+extern "C" UE_AUTORTFM_NOAUTORTFM UE_AUTORTFM_API void* autortfm_lookup_function(void* OriginalFunction, const char* Where) noexcept
 {
 	return FunctionMapLookup(OriginalFunction, Where);
 }
 
-extern "C" UE_AUTORTFM_API void autortfm_memcpy(void* Dst, const void* Src, size_t Size)
+extern "C" UE_AUTORTFM_API void autortfm_memcpy(void* Dst, const void* Src, size_t Size) noexcept
 {
 	FContext* Context = FContext::Get();
     Memcpy(Dst, Src, Size, Context);
 }
 
-extern "C" UE_AUTORTFM_API void autortfm_memmove(void* Dst, const void* Src, size_t Size)
+extern "C" UE_AUTORTFM_API void autortfm_memmove(void* Dst, const void* Src, size_t Size) noexcept
 {
 	FContext* Context = FContext::Get();
     Memmove(Dst, Src, Size, Context);
 }
 
-extern "C" UE_AUTORTFM_API void autortfm_memset(void* Dst, int Value, size_t Size)
+extern "C" UE_AUTORTFM_API void autortfm_memset(void* Dst, int Value, size_t Size) noexcept
 {
 	FContext* Context = FContext::Get();
     Memset(Dst, Value, Size, Context);
 }
 
-extern "C" UE_AUTORTFM_API void autortfm_unreachable(const char* Message)
+extern "C" UE_AUTORTFM_API void autortfm_unreachable(const char* Message) noexcept
 {
 	AUTORTFM_REPORT_ERROR("AutoRTFM Unreachable: %s", Message);
 }
 
-extern "C" UE_AUTORTFM_API void autortfm_llvm_fail(const char* Message)
+extern "C" UE_AUTORTFM_API void autortfm_llvm_fail(const char* Message) noexcept
 {
 	if (Message)
 	{
@@ -153,7 +153,7 @@ extern "C" UE_AUTORTFM_API void autortfm_llvm_fail(const char* Message)
 	}
 }
 
-extern "C" UE_AUTORTFM_API void autortfm_llvm_missing_function()
+extern "C" UE_AUTORTFM_API void autortfm_llvm_missing_function() noexcept
 {
 	if (ForTheRuntime::GetInternalAbortAction() == ForTheRuntime::EAutoRTFMInternalAbortActionState::Crash)
 	{
@@ -168,7 +168,7 @@ extern "C" UE_AUTORTFM_API void autortfm_llvm_missing_function()
     Context->AbortByLanguageAndThrow();
 }
 
-extern "C" UE_AUTORTFM_API void autortfm_called_no_autortfm()
+extern "C" UE_AUTORTFM_API void autortfm_called_no_autortfm() noexcept
 {
 	AUTORTFM_FATAL("inlined UE_AUTORTFM_NOAUTORTFM function called from the closed");
 }
diff --git a/Engine/Source/Runtime/AutoRTFM/Private/Context.cpp b/Engine/Source/Runtime/AutoRTFM/Private/Context.cpp
--- a/Engine/Source/Runtime/AutoRTFM/Private/Context.cpp
+++ b/Engine/Source/Runtime/AutoRTFM/Private/Context.cpp
@@ -444,6 +444,43 @@ ETransactionResult FContext::Transact(void (*UninstrumentedFunction)(void*), voi
 
 	if (!CurrentTransaction)
 	{
+		// If exceptions are enabled, then ensure that the transaction is automatically committed if
+		// an exception is thrown inside the transaction and the handler is outside the transaction.
+		struct FAutoCommitter final
+		{
+			FAutoCommitter(FContext& Context) : Context{Context} {}
+			~FAutoCommitter()
+			{
+				if (bCommitOnDestruct)
+				{
+					Commit();
+				}
+			}
+
+			void Commit()
+			{
+				bCommitOnDestruct = false;
+
+				if (!Context.CurrentTransaction->IsDone())
+				{
+					Context.CurrentTransaction->SetDone();
+				}
+
+				Context.PopCallNest();
+				Context.PopTransaction();
+				Context.ClearTransactionStatus();
+
+				AUTORTFM_ASSERT(Context.CurrentNest == nullptr);
+				AUTORTFM_ASSERT(Context.CurrentTransaction == nullptr);
+
+				Context.Reset();
+			}
+		private:
+			FContext& Context;
+			bool bCommitOnDestruct = true;
+		};
+		FAutoCommitter AutoCommitter(*this);
+
 		AUTORTFM_ASSERT(Status == EContextStatus::Idle);
 
 		AUTORTFM_ASSERT(CurrentThreadId == FThreadID::Invalid);
@@ -574,23 +611,52 @@ ETransactionResult FContext::Transact(void (*UninstrumentedFunction)(void*), voi
 			break;
 		}
 
-		if (!NewTransaction->IsDone())
+		AutoCommitter.Commit();
+	}
+	else
+	{
+		// This transaction is within another transaction
+
+		// If exceptions are enabled, then ensure that the transaction is automatically committed if
+		// an exception is thrown inside the transaction and the handler is outside the transaction.
+		struct FAutoCommitter final
 		{
-			NewTransaction->SetDone();
-		}
+			FAutoCommitter(FContext& Context) : Context{Context} {}
+			~FAutoCommitter()
+			{
+				if (bCommitOnDestruct)
+				{
+					Apply();
+				}
+			}
 
-		PopCallNest();
-		PopTransaction();
-		ClearTransactionStatus();
+			ETransactionResult Apply()
+			{
+				bCommitOnDestruct = false;
 
-		AUTORTFM_ASSERT(CurrentNest == nullptr);
-		AUTORTFM_ASSERT(CurrentTransaction == nullptr);
+				ETransactionResult Result = Context.ResolveNestedTransaction(Context.CurrentTransaction);
+				
+				Context.PopCallNest();
+				Context.PopTransaction();
+
+				AUTORTFM_ASSERT(Context.CurrentNest != nullptr);
+				AUTORTFM_ASSERT(Context.CurrentTransaction != nullptr);
+
+				// Cascading aborts should cause all transactions to abort!
+				if (Result == ETransactionResult::AbortedByCascade)
+				{
+					Context.CurrentTransaction->AbortAndThrow();
+				}
+
+				Context.ClearTransactionStatus();
+				return Result;
+			}
+		private:
+			FContext& Context;
+			bool bCommitOnDestruct = true;
+		};
+		FAutoCommitter AutoCommitter(*this);
 
-        Reset();
-	}
-    else
-    {
-		// This transaction is within another transaction
 		AUTORTFM_ASSERT(Status == EContextStatus::OnTrack);
 
 		AUTORTFM_ASSERT(CurrentThreadId == FThreadID::GetCurrent());
@@ -629,27 +695,10 @@ ETransactionResult FContext::Transact(void (*UninstrumentedFunction)(void*), voi
 				}
 			}
 
-			Result = ResolveNestedTransaction(NewTransaction);
-			break;
-		}
-
-		PopCallNest();
-		PopTransaction();
-
-		AUTORTFM_ASSERT(CurrentNest != nullptr);
-		AUTORTFM_ASSERT(CurrentTransaction != nullptr);
-
-		// Cascading aborts should cause all transactions to abort!
-		switch (Result)
-		{
-		default:
-			break;
-		case ETransactionResult::AbortedByCascade:
-			CurrentTransaction->AbortAndThrow();
 			break;
 		}
 
-		ClearTransactionStatus();
+		Result = AutoCommitter.Apply();
 	}
 
 	return Result;
diff --git a/Engine/Source/Runtime/AutoRTFM/Private/StdLibSupport.cpp b/Engine/Source/Runtime/AutoRTFM/Private/StdLibSupport.cpp
--- a/Engine/Source/Runtime/AutoRTFM/Private/StdLibSupport.cpp
+++ b/Engine/Source/Runtime/AutoRTFM/Private/StdLibSupport.cpp
@@ -595,6 +595,17 @@ int RTFM_atexit(void(*Callback)(void)) throw()
 	return 0;
 }
 
+#if AUTORTFM_PLATFORM_WINDOWS && AUTORTFM_EXCEPTIONS_ENABLED
+extern "C" void __stdcall _CxxThrowException(void* pExceptionObject, struct _ThrowInfo* pThrowInfo);
+#endif
+
+#if AUTORTFM_PLATFORM_LINUX && AUTORTFM_EXCEPTIONS_ENABLED
+extern "C" void __cxa_throw(void* thrown_exception, struct std::type_info * tinfo, void (*dest)(void*));
+extern "C" void* __cxa_allocate_exception(size_t thrown_size) throw();
+extern "C" void* __cxa_begin_catch(void* exceptionObject) throw();
+extern "C" void __cxa_end_catch();
+#endif
+
 }  // anonymous namespace
 
 ////////////////////////////////////////////////////////////////////////////////
@@ -807,6 +818,17 @@ UE_AUTORTFM_REGISTER_OPEN_TO_CLOSED_FUNCTIONS(
 	UE_AUTORTFM_MAP_OPEN_TO_SELF(_ZNSt3__112__next_primeEm),
 #endif
 
+#if AUTORTFM_PLATFORM_WINDOWS && AUTORTFM_EXCEPTIONS_ENABLED
+	UE_AUTORTFM_MAP_OPEN_TO_SELF(_CxxThrowException),
+#endif
+
+#if AUTORTFM_PLATFORM_LINUX && AUTORTFM_EXCEPTIONS_ENABLED
+	UE_AUTORTFM_MAP_OPEN_TO_SELF(__cxa_allocate_exception),
+	UE_AUTORTFM_MAP_OPEN_TO_SELF(__cxa_throw),
+	UE_AUTORTFM_MAP_OPEN_TO_SELF(__cxa_begin_catch),
+	UE_AUTORTFM_MAP_OPEN_TO_SELF(__cxa_end_catch),
+#endif
+
 #if AUTORTFM_PLATFORM_LINUX
 	UE_AUTORTFM_MAP_OPEN_TO_SELF(getpid),
 	UE_AUTORTFM_MAP_OPEN_TO_SELF(clock_gettime),
diff --git a/Engine/Source/Runtime/AutoRTFM/Public/AutoRTFM.h b/Engine/Source/Runtime/AutoRTFM/Public/AutoRTFM.h
--- a/Engine/Source/Runtime/AutoRTFM/Public/AutoRTFM.h
+++ b/Engine/Source/Runtime/AutoRTFM/Public/AutoRTFM.h
@@ -59,7 +59,7 @@
 #define UE_AUTORTFM_CONCAT_IMPL(A, B) A ## B
 #define UE_AUTORTFM_CONCAT(A, B) UE_AUTORTFM_CONCAT_IMPL(A, B)
 
-#if defined(__UNREAL__) && !(defined(UE_AUTORTFM_DO_NOT_INCLUDE_PLATFORM_H) && UE_AUTORTFM_DO_NOT_INCLUDE_PLATFORM_H)
+#if defined(__cplusplus) && defined(__UNREAL__) && !(defined(UE_AUTORTFM_DO_NOT_INCLUDE_PLATFORM_H) && UE_AUTORTFM_DO_NOT_INCLUDE_PLATFORM_H)
 	// Include HAL/Platform.h for DLLIMPORT / DLLEXPORT definitions, which
 	// UBT can use as a definition for AUTORTFM_API.
 	#include <HAL/Platform.h>
@@ -133,6 +133,15 @@ extern "C"
 // The C API should not be used directly - it is here purely as an
 // implementation detail.
 
+
+#ifdef __cplusplus
+#define AUTORTFM_NOEXCEPT noexcept
+#define AUTORTFM_EXCEPT noexcept(false)
+#else
+#define AUTORTFM_NOEXCEPT
+#define AUTORTFM_EXCEPT
+#endif
+
 // This must match AutoRTFM::ETransactionResult.
 typedef enum
 {
@@ -239,9 +248,9 @@ typedef struct
 // Parameters:
 //   ExternAPI - Function pointers used by AutoRTFM for heap allocations, etc.
 //               Must be non-null.
-UE_AUTORTFM_API void autortfm_initialize(const autortfm_extern_api* ExternAPI);
+UE_AUTORTFM_API void autortfm_initialize(const autortfm_extern_api* ExternAPI) AUTORTFM_NOEXCEPT;
 #else
-UE_AUTORTFM_CRITICAL_INLINE void autortfm_initialize(const autortfm_extern_api* ExternAPI)
+UE_AUTORTFM_CRITICAL_INLINE void autortfm_initialize(const autortfm_extern_api* ExternAPI) AUTORTFM_NOEXCEPT
 {
 	UE_AUTORTFM_UNUSED(ExternAPI);
 }
@@ -250,36 +259,36 @@ UE_AUTORTFM_CRITICAL_INLINE void autortfm_initialize(const autortfm_extern_api*
 #if UE_AUTORTFM_ENABLED
 // Note: There is no implementation of this function.
 // The AutoRTFM compiler will replace all calls to this function with a constant boolean value.
-UE_AUTORTFM_API bool autortfm_is_closed(void);
+UE_AUTORTFM_API bool autortfm_is_closed(void) AUTORTFM_NOEXCEPT;
 #else
-UE_AUTORTFM_CRITICAL_INLINE bool autortfm_is_closed(void)
+UE_AUTORTFM_CRITICAL_INLINE bool autortfm_is_closed(void) AUTORTFM_NOEXCEPT
 {
     return false;
 }
 #endif
 
 #if UE_AUTORTFM_ENABLED
-UE_AUTORTFM_API bool autortfm_is_transactional(void);
+UE_AUTORTFM_API bool autortfm_is_transactional(void) AUTORTFM_NOEXCEPT;
 #else
-UE_AUTORTFM_CRITICAL_INLINE bool autortfm_is_transactional(void)
+UE_AUTORTFM_CRITICAL_INLINE bool autortfm_is_transactional(void) AUTORTFM_NOEXCEPT
 {
     return false;
 }
 #endif
 
 #if UE_AUTORTFM_ENABLED
-UE_AUTORTFM_API bool autortfm_is_committing_or_aborting(void);
+UE_AUTORTFM_API bool autortfm_is_committing_or_aborting(void) AUTORTFM_NOEXCEPT;
 #else
-UE_AUTORTFM_CRITICAL_INLINE bool autortfm_is_committing_or_aborting(void)
+UE_AUTORTFM_CRITICAL_INLINE bool autortfm_is_committing_or_aborting(void) AUTORTFM_NOEXCEPT
 {
 	return false;
 }
 #endif
 
 #if UE_AUTORTFM_ENABLED
-UE_AUTORTFM_API autortfm_result autortfm_transact(void (*UninstrumentedWork)(void*), void (*InstrumentedWork)(void*), void* Arg);
+UE_AUTORTFM_API autortfm_result autortfm_transact(void (*UninstrumentedWork)(void*), void (*InstrumentedWork)(void*), void* Arg) AUTORTFM_EXCEPT;
 #else
-UE_AUTORTFM_CRITICAL_INLINE autortfm_result autortfm_transact(void (*UninstrumentedWork)(void*), void (*InstrumentedWork)(void*), void* Arg)
+UE_AUTORTFM_CRITICAL_INLINE autortfm_result autortfm_transact(void (*UninstrumentedWork)(void*), void (*InstrumentedWork)(void*), void* Arg) AUTORTFM_EXCEPT
 {
 	UE_AUTORTFM_UNUSED(InstrumentedWork);
 	UninstrumentedWork(Arg);
@@ -309,45 +318,45 @@ UE_AUTORTFM_CRITICAL_INLINE void autortfm_commit(void (*UninstrumentedWork)(void
 #endif
 
 #if UE_AUTORTFM_ENABLED
-UE_AUTORTFM_API void autortfm_abort_transaction();
+UE_AUTORTFM_API void autortfm_abort_transaction() AUTORTFM_NOEXCEPT;
 #else
-UE_AUTORTFM_CRITICAL_INLINE void autortfm_abort_transaction() {}
+UE_AUTORTFM_CRITICAL_INLINE void autortfm_abort_transaction() AUTORTFM_NOEXCEPT {}
 #endif
 
 #if UE_AUTORTFM_ENABLED
-UE_AUTORTFM_API autortfm_result autortfm_rollback_transaction();
+UE_AUTORTFM_API autortfm_result autortfm_rollback_transaction() AUTORTFM_NOEXCEPT;
 #else
-UE_AUTORTFM_CRITICAL_INLINE autortfm_result autortfm_rollback_transaction() { return autortfm_aborted_by_request; }
+UE_AUTORTFM_CRITICAL_INLINE autortfm_result autortfm_rollback_transaction() AUTORTFM_NOEXCEPT { return autortfm_aborted_by_request; }
 #endif
 
 #if UE_AUTORTFM_ENABLED
-UE_AUTORTFM_API autortfm_result autortfm_cascading_rollback_transaction();
+UE_AUTORTFM_API autortfm_result autortfm_cascading_rollback_transaction() AUTORTFM_NOEXCEPT;
 #else
-UE_AUTORTFM_CRITICAL_INLINE autortfm_result autortfm_cascading_rollback_transaction() { return autortfm_aborted_by_cascade; }
+UE_AUTORTFM_CRITICAL_INLINE autortfm_result autortfm_cascading_rollback_transaction() AUTORTFM_NOEXCEPT { return autortfm_aborted_by_cascade; }
 #endif
 
 #if UE_AUTORTFM_ENABLED
-UE_AUTORTFM_API bool autortfm_start_transaction();
+UE_AUTORTFM_API bool autortfm_start_transaction() AUTORTFM_NOEXCEPT;
 #else
-UE_AUTORTFM_CRITICAL_INLINE bool autortfm_start_transaction() { return false; }
+UE_AUTORTFM_CRITICAL_INLINE bool autortfm_start_transaction() AUTORTFM_NOEXCEPT { return false; }
 #endif
 
 #if UE_AUTORTFM_ENABLED
-UE_AUTORTFM_API autortfm_result autortfm_commit_transaction();
+UE_AUTORTFM_API autortfm_result autortfm_commit_transaction() AUTORTFM_NOEXCEPT;
 #else
-UE_AUTORTFM_CRITICAL_INLINE autortfm_result autortfm_commit_transaction() { return autortfm_aborted_by_language; }
+UE_AUTORTFM_CRITICAL_INLINE autortfm_result autortfm_commit_transaction() AUTORTFM_NOEXCEPT { return autortfm_aborted_by_language; }
 #endif
 
 #if UE_AUTORTFM_ENABLED
-UE_AUTORTFM_API void autortfm_clear_transaction_status();
+UE_AUTORTFM_API void autortfm_clear_transaction_status() AUTORTFM_NOEXCEPT;
 #else
-UE_AUTORTFM_CRITICAL_INLINE void autortfm_clear_transaction_status() {}
+UE_AUTORTFM_CRITICAL_INLINE void autortfm_clear_transaction_status() AUTORTFM_NOEXCEPT {}
 #endif
 
 #if UE_AUTORTFM_ENABLED
-UE_AUTORTFM_API autortfm_status autortfm_get_context_status();
+UE_AUTORTFM_API autortfm_status autortfm_get_context_status() AUTORTFM_NOEXCEPT;
 #else
-UE_AUTORTFM_CRITICAL_INLINE autortfm_status autortfm_get_context_status() { return autortfm_status::autortfm_status_idle; }
+UE_AUTORTFM_CRITICAL_INLINE autortfm_status autortfm_get_context_status() AUTORTFM_NOEXCEPT { return autortfm_status::autortfm_status_idle; }
 #endif
 
 #if UE_AUTORTFM_ENABLED
@@ -378,9 +387,9 @@ AUTORTFM_RESTORE_UNREACHABLE_CODE_WARNINGS
 #endif
 
 #if UE_AUTORTFM_ENABLED
-UE_AUTORTFM_API void autortfm_record_open_write(void* Ptr, size_t Size);
+UE_AUTORTFM_API void autortfm_record_open_write(void* Ptr, size_t Size) AUTORTFM_NOEXCEPT;
 #else
-UE_AUTORTFM_CRITICAL_INLINE void autortfm_record_open_write(void* Ptr, size_t Size)
+UE_AUTORTFM_CRITICAL_INLINE void autortfm_record_open_write(void* Ptr, size_t Size) AUTORTFM_NOEXCEPT
 {
 	UE_AUTORTFM_UNUSED(Ptr);
 	UE_AUTORTFM_UNUSED(Size);
@@ -388,9 +397,9 @@ UE_AUTORTFM_CRITICAL_INLINE void autortfm_record_open_write(void* Ptr, size_t Si
 #endif
 
 #if UE_AUTORTFM_ENABLED
-UE_AUTORTFM_API void autortfm_record_open_write_no_memory_validation(void* Ptr, size_t Size);
+UE_AUTORTFM_API void autortfm_record_open_write_no_memory_validation(void* Ptr, size_t Size) AUTORTFM_NOEXCEPT;
 #else
-UE_AUTORTFM_CRITICAL_INLINE void autortfm_record_open_write_no_memory_validation(void* Ptr, size_t Size)
+UE_AUTORTFM_CRITICAL_INLINE void autortfm_record_open_write_no_memory_validation(void* Ptr, size_t Size) AUTORTFM_NOEXCEPT
 {
 	UE_AUTORTFM_UNUSED(Ptr);
 	UE_AUTORTFM_UNUSED(Size);
@@ -410,33 +419,33 @@ struct autortfm_open_to_closed_mapping
 struct autortfm_open_to_closed_table
 {
 	// Null-terminated open function to closed function mapping table.
-	const autortfm_open_to_closed_mapping* Mappings;
+	const struct autortfm_open_to_closed_mapping* Mappings;
 	// An intrusive linked-list pointer to the previous autortfm_open_to_closed_table.
 	// Used by autortfm_register_open_to_closed_functions().
-	autortfm_open_to_closed_table* Prev;
+	struct autortfm_open_to_closed_table* Prev;
 	// An intrusive linked-list pointer to the next autortfm_open_to_closed_table.
 	// Used by autortfm_register_open_to_closed_functions().
-	autortfm_open_to_closed_table* Next;
+	struct autortfm_open_to_closed_table* Next;
 };
 
 #if UE_AUTORTFM_ENABLED
-UE_AUTORTFM_API void autortfm_register_open_to_closed_functions(autortfm_open_to_closed_table* Table);
-UE_AUTORTFM_API void autortfm_unregister_open_to_closed_functions(autortfm_open_to_closed_table* Table);
+UE_AUTORTFM_API void autortfm_register_open_to_closed_functions(struct autortfm_open_to_closed_table* Table) AUTORTFM_NOEXCEPT;
+UE_AUTORTFM_API void autortfm_unregister_open_to_closed_functions(struct autortfm_open_to_closed_table* Table) AUTORTFM_NOEXCEPT;
 #else
-UE_AUTORTFM_CRITICAL_INLINE void autortfm_register_open_to_closed_functions(autortfm_open_to_closed_table* Table)
+UE_AUTORTFM_CRITICAL_INLINE void autortfm_register_open_to_closed_functions(autortfm_open_to_closed_table* Table) AUTORTFM_NOEXCEPT
 {
 	UE_AUTORTFM_UNUSED(Table);
 }
-UE_AUTORTFM_CRITICAL_INLINE void autortfm_unregister_open_to_closed_functions(autortfm_open_to_closed_table* Table)
+UE_AUTORTFM_CRITICAL_INLINE void autortfm_unregister_open_to_closed_functions(autortfm_open_to_closed_table* Table) AUTORTFM_NOEXCEPT
 {
 	UE_AUTORTFM_UNUSED(Table);
 }
 #endif
 
 #if UE_AUTORTFM_ENABLED
-UE_AUTORTFM_API bool autortfm_is_on_current_transaction_stack(void* Ptr);
+UE_AUTORTFM_API bool autortfm_is_on_current_transaction_stack(void* Ptr) AUTORTFM_NOEXCEPT;
 #else
-UE_AUTORTFM_CRITICAL_INLINE bool autortfm_is_on_current_transaction_stack(void* Ptr)
+UE_AUTORTFM_CRITICAL_INLINE bool autortfm_is_on_current_transaction_stack(void* Ptr) AUTORTFM_NOEXCEPT
 {
 	UE_AUTORTFM_UNUSED(Ptr);
 	return false;
@@ -453,9 +462,9 @@ UE_AUTORTFM_CRITICAL_INLINE void autortfm_on_commit(void (*work)(void* arg), voi
 #endif
 
 #if UE_AUTORTFM_ENABLED
-UE_AUTORTFM_API void autortfm_on_abort(void (*work)(void* arg), void* arg);
+UE_AUTORTFM_API void autortfm_on_abort(void (*work)(void* arg), void* arg) AUTORTFM_NOEXCEPT;
 #else
-UE_AUTORTFM_CRITICAL_INLINE void autortfm_on_abort(void (*work)(void* arg), void* arg)
+UE_AUTORTFM_CRITICAL_INLINE void autortfm_on_abort(void (*work)(void* arg), void* arg) AUTORTFM_NOEXCEPT
 {
 	UE_AUTORTFM_UNUSED(work);
 	UE_AUTORTFM_UNUSED(arg);
@@ -463,26 +472,26 @@ UE_AUTORTFM_CRITICAL_INLINE void autortfm_on_abort(void (*work)(void* arg), void
 #endif
 
 #if UE_AUTORTFM_ENABLED
-UE_AUTORTFM_API void autortfm_push_on_abort_handler(const void* key, void (*work)(void* arg), void* arg);
-UE_AUTORTFM_API void autortfm_pop_on_abort_handler(const void* key);
+UE_AUTORTFM_API void autortfm_push_on_abort_handler(const void* key, void (*work)(void* arg), void* arg) AUTORTFM_NOEXCEPT;
+UE_AUTORTFM_API void autortfm_pop_on_abort_handler(const void* key) AUTORTFM_NOEXCEPT;
 #else
-UE_AUTORTFM_CRITICAL_INLINE void autortfm_push_on_abort_handler(const void* key, void (*work)(void* arg), void* arg)
+UE_AUTORTFM_CRITICAL_INLINE void autortfm_push_on_abort_handler(const void* key, void (*work)(void* arg), void* arg) AUTORTFM_NOEXCEPT
 {
 	UE_AUTORTFM_UNUSED(key);
 	UE_AUTORTFM_UNUSED(work);
 	UE_AUTORTFM_UNUSED(arg);
 }
 
-UE_AUTORTFM_CRITICAL_INLINE void autortfm_pop_on_abort_handler(const void* key)
+UE_AUTORTFM_CRITICAL_INLINE void autortfm_pop_on_abort_handler(const void* key) AUTORTFM_NOEXCEPT
 {
 	UE_AUTORTFM_UNUSED(key);
 }
 #endif
 
 #if UE_AUTORTFM_ENABLED
-UE_AUTORTFM_API void* autortfm_did_allocate(void* ptr, size_t size);
+UE_AUTORTFM_API void* autortfm_did_allocate(void* ptr, size_t size) AUTORTFM_NOEXCEPT;
 #else
-UE_AUTORTFM_CRITICAL_INLINE void* autortfm_did_allocate(void* ptr, size_t size)
+UE_AUTORTFM_CRITICAL_INLINE void* autortfm_did_allocate(void* ptr, size_t size) AUTORTFM_NOEXCEPT
 {
 	UE_AUTORTFM_UNUSED(size);
     return ptr;
@@ -490,9 +499,9 @@ UE_AUTORTFM_CRITICAL_INLINE void* autortfm_did_allocate(void* ptr, size_t size)
 #endif
 
 #if UE_AUTORTFM_ENABLED
-UE_AUTORTFM_API void autortfm_did_free(void* ptr);
+UE_AUTORTFM_API void autortfm_did_free(void* ptr) AUTORTFM_NOEXCEPT;
 #else
-UE_AUTORTFM_CRITICAL_INLINE void autortfm_did_free(void* ptr)
+UE_AUTORTFM_CRITICAL_INLINE void autortfm_did_free(void* ptr) AUTORTFM_NOEXCEPT
 {
 	UE_AUTORTFM_UNUSED(ptr);
 }
@@ -505,9 +514,9 @@ UE_AUTORTFM_CRITICAL_INLINE void autortfm_did_free(void* ptr)
 // be injected by the compiler into a global constructor in the AutoRTFM compiled
 // code.
 #if UE_AUTORTFM_ENABLED
-UE_AUTORTFM_API void autortfm_check_abi(void* ptr, size_t size);
+UE_AUTORTFM_API void autortfm_check_abi(void* ptr, size_t size) AUTORTFM_NOEXCEPT;
 #else
-UE_AUTORTFM_CRITICAL_INLINE void autortfm_check_abi(void* ptr, size_t size)
+UE_AUTORTFM_CRITICAL_INLINE void autortfm_check_abi(void* ptr, size_t size) AUTORTFM_NOEXCEPT
 {
 	UE_AUTORTFM_UNUSED(ptr);
 	UE_AUTORTFM_UNUSED(size);
@@ -515,9 +524,9 @@ UE_AUTORTFM_CRITICAL_INLINE void autortfm_check_abi(void* ptr, size_t size)
 #endif
 
 #if UE_AUTORTFM_ENABLED
-UE_AUTORTFM_API void autortfm_unreachable(const char* Message);
+UE_AUTORTFM_API void autortfm_unreachable(const char* Message) AUTORTFM_NOEXCEPT;
 #else
-UE_AUTORTFM_CRITICAL_INLINE void autortfm_unreachable(const char* Message)
+UE_AUTORTFM_CRITICAL_INLINE void autortfm_unreachable(const char* Message) AUTORTFM_NOEXCEPT
 {
 	UE_AUTORTFM_UNUSED(Message);
 }
@@ -635,7 +644,7 @@ template<typename FunctorType>
 void AutoRTFMFunctorInvoker(void* Arg) { (*static_cast<const FunctorType*>(Arg))(); }
 
 #if UE_AUTORTFM
-extern "C" UE_AUTORTFM_API void* autortfm_lookup_function(void* OriginalFunction, const char* Where);
+extern "C" UE_AUTORTFM_API void* autortfm_lookup_function(void* OriginalFunction, const char* Where) AUTORTFM_NOEXCEPT;
 
 template<typename FunctorType>
 auto AutoRTFMLookupInstrumentedFunctorInvoker(const FunctorType& Functor) -> void(*)(void*)
