[AutoRTFM] Add transactional support for Fls(Get|Set)Value, fix TlsSetValue.

TlsSetIndex was previously broken from AutoRTFM and would not rollback correctly in an abort. The call to `AutoRTFM::OnAbort` from `RTFM_TlsSetIndex` was ignored because the handler is running from the open. This is now fixed by adding new ForTheRuntime methods which allow us to register Abort and Commit handlers from the open.

The fiber local storage functions can be used by EngineTest when logging. Since AutoRTFM transactions only run on a single fiber, these are pretty trivial to support. The implementation is essentially identical to TlsSetValue.

diff --git a/Engine/Source/Programs/AutoRTFMTests/Private/LocalStorageTests.cpp b/Engine/Source/Programs/AutoRTFMTests/Private/LocalStorageTests.cpp
new file mode 100644
--- /dev/null
+++ b/Engine/Source/Programs/AutoRTFMTests/Private/LocalStorageTests.cpp
@@ -0,0 +1,78 @@
+// Copyright Epic Games, Inc. All Rights Reserved.
+
+#include "Catch2Includes.h"
+#include "AutoRTFM.h"
+#include "AutoRTFMTesting.h"
+
+#if PLATFORM_WINDOWS
+
+#include "WindowsHeader.h"
+
+TEST_CASE("LocalStorage.TLS")
+{
+	DWORD TlsSlot = TlsAlloc();
+	int One = 1, Two = 2;
+
+	SECTION("Commit")
+	{
+		TlsSetValue(TlsSlot, &One);
+
+		AutoRTFM::Testing::Commit([&]
+		{
+			TlsSetValue(TlsSlot, &Two);
+		});
+
+		REQUIRE(TlsGetValue(TlsSlot) == &Two);
+	}
+
+	SECTION("Free")
+	{
+		TlsSetValue(TlsSlot, &One);
+
+		AutoRTFM::Testing::Abort([&]
+		{
+			TlsSetValue(TlsSlot, &Two);
+			AutoRTFM::AbortTransaction();
+		});
+
+		REQUIRE(TlsGetValue(TlsSlot) == &One);
+	}
+
+	TlsFree(TlsSlot);
+}
+
+TEST_CASE("LocalStorage.FLS")
+{
+	// This is identical to the TLS test, except `Tls` has been replaced with `Fls`.
+	DWORD FlsSlot = FlsAlloc(/*lpCallback=*/ nullptr);
+	int One = 1, Two = 2;
+
+	SECTION("Commit")
+	{
+		FlsSetValue(FlsSlot, &One);
+
+		AutoRTFM::Testing::Commit([&]
+		{
+			FlsSetValue(FlsSlot, &Two);
+		});
+
+		REQUIRE(FlsGetValue(FlsSlot) == &Two);
+	}
+
+	SECTION("Free")
+	{
+		FlsSetValue(FlsSlot, &One);
+
+		AutoRTFM::Testing::Abort([&]
+		{
+			FlsSetValue(FlsSlot, &Two);
+			AutoRTFM::AbortTransaction();
+		});
+
+		REQUIRE(FlsGetValue(FlsSlot) == &One);
+	}
+
+	FlsFree(FlsSlot);
+}
+
+#endif
\ No newline at end of file
diff --git a/Engine/Source/Runtime/AutoRTFM/Private/API.cpp b/Engine/Source/Runtime/AutoRTFM/Private/API.cpp
--- a/Engine/Source/Runtime/AutoRTFM/Private/API.cpp
+++ b/Engine/Source/Runtime/AutoRTFM/Private/API.cpp
@@ -1125,6 +1125,18 @@ void ForTheRuntime::PopAllOnAbortHandlersInternal(const void* Key)
 	}
 }
 
+UE_AUTORTFM_NOAUTORTFM
+void ForTheRuntime::RegisterOnCommitFromTheOpen(TTask<void()> && Work)
+{
+	RTFM_OnCommitInternal(std::move(Work));
+}
+
+UE_AUTORTFM_NOAUTORTFM
+void ForTheRuntime::RegisterOnAbortFromTheOpen(TTask<void()> && Work)
+{
+	RTFM_OnAbortInternal(std::move(Work));
+}
+
 extern "C" void autortfm_on_commit(void (*Work)(void*), void* Arg)
 {
 	if (autortfm_is_closed())
diff --git a/Engine/Source/Runtime/AutoRTFM/Private/StdLibSupport.cpp b/Engine/Source/Runtime/AutoRTFM/Private/StdLibSupport.cpp
--- a/Engine/Source/Runtime/AutoRTFM/Private/StdLibSupport.cpp
+++ b/Engine/Source/Runtime/AutoRTFM/Private/StdLibSupport.cpp
@@ -259,16 +259,28 @@ int RTFM___stdio_common_vfwprintf(
 
 BOOL RTFM_TlsSetValue(DWORD dwTlsIndex, LPVOID lpTlsValue)
 {
-	LPVOID CurrentValue = TlsGetValue(dwTlsIndex);
+	LPVOID OriginalValue = TlsGetValue(dwTlsIndex);
 
-	AutoRTFM::OnAbort([dwTlsIndex, CurrentValue]
+	AutoRTFM::ForTheRuntime::RegisterOnAbortFromTheOpen([dwTlsIndex, OriginalValue]
 	{
-		TlsSetValue(dwTlsIndex, CurrentValue);
+		TlsSetValue(dwTlsIndex, OriginalValue);
 	});
 
 	return TlsSetValue(dwTlsIndex, lpTlsValue);
 }
 
+BOOL RTFM_FlsSetValue(DWORD dwFlsIndex, LPVOID lpFlsValue)
+{
+	LPVOID OriginalValue = FlsGetValue(dwFlsIndex);
+
+	AutoRTFM::ForTheRuntime::RegisterOnAbortFromTheOpen([dwFlsIndex, OriginalValue]
+	{
+		FlsSetValue(dwFlsIndex, OriginalValue);
+	});
+
+	return FlsSetValue(dwFlsIndex, lpFlsValue);
+}
+
 // ^^^^^^^^^^^^^^^^^^^^^^^^^^^^^^ AUTORTFM_PLATFORM_WINDOWS ^^^^^^^^^^^^^^^^^^^^^^^^^^^^^
 #else
 // vvvvvvvvvvvvvvvvvvvvvvvvvvvvv !AUTORTFM_PLATFORM_WINDOWS vvvvvvvvvvvvvvvvvvvvvvvvvvvvv
@@ -809,9 +821,11 @@ UE_AUTORTFM_REGISTER_OPEN_TO_CLOSED_FUNCTIONS(
 	UE_AUTORTFM_MAP_OPEN_TO_SELF(GetCurrentThreadId),
 	UE_AUTORTFM_MAP_OPEN_TO_SELF(GetCurrentProcessId),
 	UE_AUTORTFM_MAP_OPEN_TO_SELF(TlsGetValue),
+	UE_AUTORTFM_MAP_OPEN_TO_SELF(FlsGetValue),
 	UE_AUTORTFM_MAP_OPEN_TO_SELF(GetLocalTime),
 	UE_AUTORTFM_MAP_OPEN_TO_SELF(GetFileAttributesW),
 	UE_AUTORTFM_MAP_OPEN_TO_CLOSED(TlsSetValue, RTFM_TlsSetValue),
+	UE_AUTORTFM_MAP_OPEN_TO_CLOSED(FlsSetValue, RTFM_FlsSetValue),
 	UE_AUTORTFM_MAP_OPEN_TO_SELF(_Query_perf_frequency),
 	UE_AUTORTFM_MAP_OPEN_TO_SELF(_Query_perf_counter),
 #else
diff --git a/Engine/Source/Runtime/AutoRTFM/Public/AutoRTFM.h b/Engine/Source/Runtime/AutoRTFM/Public/AutoRTFM.h
--- a/Engine/Source/Runtime/AutoRTFM/Public/AutoRTFM.h
+++ b/Engine/Source/Runtime/AutoRTFM/Public/AutoRTFM.h
@@ -637,6 +637,8 @@ namespace ForTheRuntime
 	UE_AUTORTFM_API void PushOnAbortHandlerInternal(const void* Key, TTask<void()>&& Work);
 	UE_AUTORTFM_API void PopOnAbortHandlerInternal(const void* Key);
 	UE_AUTORTFM_API void PopAllOnAbortHandlersInternal(const void* Key);
+	UE_AUTORTFM_API void RegisterOnCommitFromTheOpen(TTask<void()>&& Work);
+	UE_AUTORTFM_API void RegisterOnAbortFromTheOpen(TTask<void()>&& Work);
 } // namespace ForTheRuntime
 #endif
 
