[AutoRTFM] Expand memory validation to include writes made by on-abort / on-commit handlers.

The on-abort handler function, destructor and on-commit handler destructor can all execute logic in the open, which may modify memory tracked by the outer transaction.
Modification of this memory will likely cause memory corruption.

diff --git a/Engine/Source/Programs/AutoRTFMTests/Private/Abort.cpp b/Engine/Source/Programs/AutoRTFMTests/Private/Abort.cpp
--- a/Engine/Source/Programs/AutoRTFMTests/Private/Abort.cpp
+++ b/Engine/Source/Programs/AutoRTFMTests/Private/Abort.cpp
@@ -1,5 +1,6 @@
 // Copyright Epic Games, Inc. All Rights Reserved.
 
+#include "AutoRTFM.h"
 #include "AutoRTFMTesting.h"
 #include "AutoRTFMTestUtils.h"
 #include "Catch2Includes.h"
@@ -1009,3 +1010,72 @@ TEST_CASE("Abort.PushOnAbortOrdering")
 		});
 	}
 }
+
+TEST_CASE("Abort.CallbackModifiesTransactionalMemory")
+{
+	AUTORTFM_SCOPED_ENABLE_MEMORY_VALIDATION_AS_WARNING();
+	AutoRTFMTestUtils::FCaptureWarningContext WarningContext;
+
+	int Memory = 0;
+	
+	SECTION("Task")
+	{
+		AutoRTFM::Testing::Commit([&]
+		{
+			Memory = 10;
+			AutoRTFM::Testing::Abort([&]
+			{
+				AutoRTFM::OnAbort([&]
+				{
+					Memory = 20;
+				});
+				AutoRTFM::AbortTransaction();
+			});
+		});
+
+		REQUIRE(WarningContext.HasWarningSubstring(AutoRTFMTestUtils::kMemoryModifiedInAbortHandlerWarning));
+	}
+
+	SECTION("Destructor")
+	{
+		struct IncrementOnDtor
+		{
+			~IncrementOnDtor() { (*Ptr)++; }
+			int* const Ptr;
+		};
+		
+		SECTION("OnAbort")
+		{
+			AutoRTFM::Testing::Commit([&]
+			{
+				Memory = 10;
+				AutoRTFM::Testing::Abort([&]
+				{
+					AutoRTFM::OnAbort([Dtor = IncrementOnDtor{&Memory}]
+					{
+					});
+					AutoRTFM::AbortTransaction();
+				});
+			});
+
+			REQUIRE(WarningContext.HasWarningSubstring(AutoRTFMTestUtils::kMemoryModifiedInAbortHandlerDtorWarning));
+		}
+
+		SECTION("OnCommit")
+		{
+			AutoRTFM::Testing::Commit([&]
+			{
+				Memory = 10;
+				AutoRTFM::Testing::Abort([&]
+				{
+					AutoRTFM::OnCommit([Dtor = IncrementOnDtor{&Memory}]
+					{
+					});
+					AutoRTFM::AbortTransaction();
+				});
+			});
+
+			REQUIRE(WarningContext.HasWarningSubstring(AutoRTFMTestUtils::kMemoryModifiedInCommitHandlerDtorWarning));
+		}
+	}
+}
diff --git a/Engine/Source/Programs/AutoRTFMTests/Private/AutoRTFMTestUtils.h b/Engine/Source/Programs/AutoRTFMTests/Private/AutoRTFMTestUtils.h
--- a/Engine/Source/Programs/AutoRTFMTests/Private/AutoRTFMTestUtils.h
+++ b/Engine/Source/Programs/AutoRTFMTests/Private/AutoRTFMTestUtils.h
@@ -3,8 +3,13 @@
 #pragma once
 
 #include "AutoRTFM.h"
+#include "AutoRTFMDefines.h"
+#include "Containers/StringConv.h"
+#include "HAL/Platform.h"
 #include "Misc/FeedbackContext.h"
 
+#include <algorithm>
+#include <string>
 #include <vector>
 
 namespace AutoRTFMTestUtils
@@ -34,11 +39,12 @@ public:
 	FCaptureWarningContext() : OldContext(GWarn) { GWarn = this; }
 	~FCaptureWarningContext() { GWarn = OldContext; }
 
+	UE_AUTORTFM_ALWAYS_OPEN
 	void Serialize(const TCHAR* V, ELogVerbosity::Type Verbosity, const FName& Category) override
 	{
 		if (Verbosity == ELogVerbosity::Warning)
 		{
-			Warnings.push_back(V);
+			Warnings.push_back(std::string(StringCast<ANSICHAR>(V).Get()));
 		}
 		else
 		{
@@ -46,11 +52,12 @@ public:
 		}
 	}
 
+	UE_AUTORTFM_ALWAYS_OPEN
 	void Serialize(const TCHAR* V, ELogVerbosity::Type Verbosity, const FName& Category, double Time) override
 	{
 		if (Verbosity == ELogVerbosity::Warning)
 		{
-			Warnings.push_back(V);
+			Warnings.push_back(std::string(StringCast<ANSICHAR>(V).Get()));
 		}
 		else
 		{
@@ -58,17 +65,36 @@ public:
 		}
 	}
 
-	const std::vector<FString>& GetWarnings() const
+	UE_AUTORTFM_ALWAYS_OPEN
+	const std::vector<std::string>& GetWarnings() const
 	{
 		return Warnings;
 	}
 
+	UE_AUTORTFM_ALWAYS_OPEN
+	bool HasWarning(std::string_view Expected) const
+	{
+		return std::any_of(Warnings.begin(), Warnings.end(), [&](std::string_view Warning)
+		{
+			return Warning == Expected;
+		});
+	}
+
+	UE_AUTORTFM_ALWAYS_OPEN
+	bool HasWarningSubstring(std::string_view Substr) const
+	{
+		return std::any_of(Warnings.begin(), Warnings.end(), [&](std::string_view Warning)
+		{
+			return Warning.find(Substr) != std::string::npos;
+		});
+	}
+
 private:
 	FCaptureWarningContext(const FCaptureWarningContext&) = delete;
 	FCaptureWarningContext& operator = (const FCaptureWarningContext&) = delete;
 
 	FFeedbackContext* OldContext = nullptr;
-	std::vector<FString> Warnings;
+	std::vector<std::string> Warnings;
 };
 
 // A helper class that temporarily changes the memory validation level for the lifetime of the
@@ -103,10 +129,19 @@ private:
 		{AutoRTFM::EMemoryValidationLevel::Warn}; \
 	static_assert(true) /* require semicolon */
 
-static constexpr const TCHAR* kMemoryModifiedWarning =
-    TEXT("Memory modified in a transaction was also modified in an call to AutoRTFM::Open(). ")
-    TEXT("This may lead to memory corruption if the transaction is aborted.");
+static constexpr const std::string_view kMemoryModifiedInOpenWarning =
+	"Memory modified in a transaction was also modified in an call to AutoRTFM::Open().\n"
+	"This may lead to memory corruption if the transaction is aborted.";
+
+static constexpr const std::string_view kMemoryModifiedInAbortHandlerWarning =
+	"Memory modified in a transaction was also modified in the on-abort handler:";
+
+static constexpr const std::string_view kMemoryModifiedInAbortHandlerDtorWarning =
+	"Memory modified in a transaction was also modified in the destructor of the on-abort handler:";
 
+static constexpr const std::string_view kMemoryModifiedInCommitHandlerDtorWarning =
+	"Memory modified in a transaction was also modified in the destructor of the on-commit handler:";
+	
 class FScopedEnsureOnInternalAbort final
 {
 public:
diff --git a/Engine/Source/Programs/AutoRTFMTests/Private/Open.cpp b/Engine/Source/Programs/AutoRTFMTests/Private/Open.cpp
--- a/Engine/Source/Programs/AutoRTFMTests/Private/Open.cpp
+++ b/Engine/Source/Programs/AutoRTFMTests/Private/Open.cpp
@@ -442,8 +442,7 @@ TEST_CASE("Open.Collision")
             AutoRTFM::Open([&] { I = 24; });
         });
 
-        REQUIRE_THAT(WarningContext.GetWarnings(),
-            Catch::Matchers::VectorContains(FString(AutoRTFMTestUtils::kMemoryModifiedWarning)));
+        REQUIRE(WarningContext.HasWarning(AutoRTFMTestUtils::kMemoryModifiedInOpenWarning));
     }
 
     SECTION("Transact(Assign, Transact(Open(Assign)))")
@@ -457,8 +456,7 @@ TEST_CASE("Open.Collision")
             });
         });
 
-        REQUIRE_THAT(WarningContext.GetWarnings(),
-            Catch::Matchers::VectorContains(FString(AutoRTFMTestUtils::kMemoryModifiedWarning)));
+        REQUIRE(WarningContext.HasWarning(AutoRTFMTestUtils::kMemoryModifiedInOpenWarning));
     }
 
     SECTION("Transact(Assign, Transact(Transact(Open(Assign))))")
@@ -478,8 +476,7 @@ TEST_CASE("Open.Collision")
             });
         });
 
-        REQUIRE_THAT(WarningContext.GetWarnings(),
-            Catch::Matchers::VectorContains(FString(AutoRTFMTestUtils::kMemoryModifiedWarning)));
+        REQUIRE(WarningContext.HasWarning(AutoRTFMTestUtils::kMemoryModifiedInOpenWarning));
     }
 
     SECTION("Transact(Transact(Assign, Transact(Open(Assign))))")
@@ -499,8 +496,7 @@ TEST_CASE("Open.Collision")
             });
         });
 
-        REQUIRE_THAT(WarningContext.GetWarnings(),
-            Catch::Matchers::VectorContains(FString(AutoRTFMTestUtils::kMemoryModifiedWarning)));
+        REQUIRE(WarningContext.HasWarning(AutoRTFMTestUtils::kMemoryModifiedInOpenWarning));
     }
 
     SECTION("Transact(Transact(Assign, Open(Assign, Transact())))")
@@ -515,8 +511,7 @@ TEST_CASE("Open.Collision")
             });
         });
 
-        REQUIRE_THAT(WarningContext.GetWarnings(),
-            Catch::Matchers::VectorContains(FString(AutoRTFMTestUtils::kMemoryModifiedWarning)));
+        REQUIRE(WarningContext.HasWarning(AutoRTFMTestUtils::kMemoryModifiedInOpenWarning));
     }
 
     SECTION("Transact(Transact(Assign, Open(Transact(), Assign)))")
@@ -531,8 +526,7 @@ TEST_CASE("Open.Collision")
             });
         });
 
-        REQUIRE_THAT(WarningContext.GetWarnings(),
-            Catch::Matchers::VectorContains(FString(AutoRTFMTestUtils::kMemoryModifiedWarning)));
+        REQUIRE(WarningContext.HasWarning(AutoRTFMTestUtils::kMemoryModifiedInOpenWarning));
     }
 
     SECTION("Transact(Assign, CallOpen(Assign))")
@@ -543,8 +537,7 @@ TEST_CASE("Open.Collision")
             AssignIntPointer(&I, 24);
         });
 
-        REQUIRE_THAT(WarningContext.GetWarnings(),
-            Catch::Matchers::VectorContains(FString(AutoRTFMTestUtils::kMemoryModifiedWarning)));
+        REQUIRE(WarningContext.HasWarning(AutoRTFMTestUtils::kMemoryModifiedInOpenWarning));
     }
 
     SECTION("Transact(Assign, Transact(CallOpen(Assign)))")
@@ -558,8 +551,7 @@ TEST_CASE("Open.Collision")
             });
         });
 
-        REQUIRE_THAT(WarningContext.GetWarnings(),
-            Catch::Matchers::VectorContains(FString(AutoRTFMTestUtils::kMemoryModifiedWarning)));
+        REQUIRE(WarningContext.HasWarning(AutoRTFMTestUtils::kMemoryModifiedInOpenWarning));
     }
 
     SECTION("Transact(Assign, OpenNoValidation(Assign))")
@@ -608,8 +600,7 @@ TEST_CASE("Open.Collision")
             });
         });
 
-        REQUIRE_THAT(WarningContext.GetWarnings(),
-            Catch::Matchers::VectorContains(FString(AutoRTFMTestUtils::kMemoryModifiedWarning)));
+        REQUIRE(WarningContext.HasWarning(AutoRTFMTestUtils::kMemoryModifiedInOpenWarning));
     }
     
     SECTION("Transact(Assign, Open(Assign, Close(OpenNoValidation())))")
@@ -627,8 +618,7 @@ TEST_CASE("Open.Collision")
             });
         });
 
-        REQUIRE_THAT(WarningContext.GetWarnings(),
-            Catch::Matchers::VectorContains(FString(AutoRTFMTestUtils::kMemoryModifiedWarning)));
+        REQUIRE(WarningContext.HasWarning(AutoRTFMTestUtils::kMemoryModifiedInOpenWarning));
     }
 
     SECTION("Transact(Open(RecordOpenWrite()), Open(Assign))")
@@ -646,8 +636,7 @@ TEST_CASE("Open.Collision")
         });
 
         REQUIRE(I == 42);
-        REQUIRE_THAT(WarningContext.GetWarnings(),
-            Catch::Matchers::VectorContains(FString(AutoRTFMTestUtils::kMemoryModifiedWarning)));
+        REQUIRE(WarningContext.HasWarning(AutoRTFMTestUtils::kMemoryModifiedInOpenWarning));
     }
 
     SECTION("Transact(Open(RecordOpenWriteNoMemoryValidation()), Open(RecordOpenWrite()), Open(Assign))")
@@ -668,8 +657,7 @@ TEST_CASE("Open.Collision")
             });
         });
 
-        REQUIRE_THAT(WarningContext.GetWarnings(),
-            Catch::Matchers::VectorContains(FString(AutoRTFMTestUtils::kMemoryModifiedWarning)));
+        REQUIRE(WarningContext.HasWarning(AutoRTFMTestUtils::kMemoryModifiedInOpenWarning));
     }
 
     SECTION("Transact(Transact(Open(RecordOpenWriteNoMemoryValidation())), Open(RecordOpenWrite()), Open(Assign))")
@@ -693,8 +681,7 @@ TEST_CASE("Open.Collision")
             });
         });
 
-        REQUIRE_THAT(WarningContext.GetWarnings(),
-            Catch::Matchers::VectorContains(FString(AutoRTFMTestUtils::kMemoryModifiedWarning)));
+        REQUIRE(WarningContext.HasWarning(AutoRTFMTestUtils::kMemoryModifiedInOpenWarning));
     }
 
     SECTION("Transact(Open(RecordOpenWriteNoMemoryValidation()), Open(RecordOpenWrite()), Open(Assign)) <large>")
@@ -716,7 +703,6 @@ TEST_CASE("Open.Collision")
             });
         });
 
-        REQUIRE_THAT(WarningContext.GetWarnings(),
-            Catch::Matchers::VectorContains(FString(AutoRTFMTestUtils::kMemoryModifiedWarning)));
+        REQUIRE(WarningContext.HasWarning(AutoRTFMTestUtils::kMemoryModifiedInOpenWarning));
     }
 }
diff --git a/Engine/Source/Programs/AutoRTFMTests/Private/OpenAPI.cpp b/Engine/Source/Programs/AutoRTFMTests/Private/OpenAPI.cpp
--- a/Engine/Source/Programs/AutoRTFMTests/Private/OpenAPI.cpp
+++ b/Engine/Source/Programs/AutoRTFMTests/Private/OpenAPI.cpp
@@ -41,8 +41,7 @@ TEST_CASE("OpenAPI.StartAbortAndStartAgain")
 	REQUIRE(ValueB == 20);
 	REQUIRE(ValueC == 0);
 
-	REQUIRE_THAT(WarningContext.GetWarnings(),
-		Catch::Matchers::VectorContains(FString(AutoRTFMTestUtils::kMemoryModifiedWarning)));
+	REQUIRE(WarningContext.HasWarning(AutoRTFMTestUtils::kMemoryModifiedInOpenWarning));
 }
 
 TEST_CASE("OpenAPI.CommitScopedFromOpen_Illegal", "[.]")
@@ -731,8 +730,7 @@ TEST_CASE("OpenAPI.Footgun1")
 	REQUIRE(ValueA == 10);
 	REQUIRE(ValueB == 0);
 
-	REQUIRE_THAT(WarningContext.GetWarnings(),
-		Catch::Matchers::VectorContains(FString(AutoRTFMTestUtils::kMemoryModifiedWarning)));
+	REQUIRE(WarningContext.HasWarning(AutoRTFMTestUtils::kMemoryModifiedInOpenWarning));
 }
 
 TEST_CASE("OpenAPI.Footgun2")
@@ -769,8 +767,7 @@ TEST_CASE("OpenAPI.Footgun2")
 	REQUIRE(ValueB == 0);
 	REQUIRE(ValueC == 10);
 
-	REQUIRE_THAT(WarningContext.GetWarnings(),
-		Catch::Matchers::VectorContains(FString(AutoRTFMTestUtils::kMemoryModifiedWarning)));
+	REQUIRE(WarningContext.HasWarning(AutoRTFMTestUtils::kMemoryModifiedInOpenWarning));
 }
 
 #if 0
diff --git a/Engine/Source/Runtime/AutoRTFM/Private/Transaction.cpp b/Engine/Source/Runtime/AutoRTFM/Private/Transaction.cpp
--- a/Engine/Source/Runtime/AutoRTFM/Private/Transaction.cpp
+++ b/Engine/Source/Runtime/AutoRTFM/Private/Transaction.cpp
@@ -10,8 +10,10 @@
 #include "HitSet.h"
 #include "Utils.h"
 
+#include <cstdarg>
 #include <cstddef>
 #include <cstdint>
+#include <optional>
 
 namespace AutoRTFM
 {
@@ -112,35 +114,140 @@ void FTransaction::AbortWithoutThrowing()
 {
 	AUTORTFM_VERBOSE("Aborting '%hs'!", GetContextStatusName(Context->GetStatus()));
 
-    AUTORTFM_ASSERT(Context->IsAborting());
+	AUTORTFM_ASSERT(Context->IsAborting());
 
-    Stats.Collect<EStatsKind::Abort>();
-    CollectStats();
+	Stats.Collect<EStatsKind::Abort>();
+	CollectStats();
 
 	// Ensure that we enter the done state before applying the commit, as this
 	// will ensure the open-memory validation is performed before the write log
 	// is cleared.
 	SetDone();
 
-    // Call the destructors of all the OnCommit functors before undoing the transactional memory and
-    // calling the OnAbort callbacks. This is important as the callback functions may have captured
-    // variables that are depending on the allocated memory.
-    CommitTasks.Reset();
+	const EMemoryValidationLevel MemoryValidationLevel = ForTheRuntime::GetMemoryValidationLevel();
 
-    Undo();
-	AbortTasks.RemoveEachBackward([&](TTask<void()>& Task)
-    { 
-        Task();
-    });
+	// Call the destructors of all the OnCommit functors before undoing the transactional memory and
+	// calling the OnAbort callbacks. This is important as the callback functions may have captured
+	// variables that are depending on the allocated memory.
+	if (AUTORTFM_LIKELY(MemoryValidationLevel == EMemoryValidationLevel::Disabled))
+	{
+		CommitTasks.Reset();
+	}
+	else
+	{
+		ResetCommitTasksWithValidation();
+	}
 
-    if (IsNested())
-    {
+	// Revert all memory in the write-log.
+	Undo();
+
+	// Call each of the on-abort handlers.
+	if (AUTORTFM_LIKELY(MemoryValidationLevel == EMemoryValidationLevel::Disabled))
+	{
+		AbortTasks.RemoveEachBackward([&](TTask<void()>& Task)
+		{ 
+			Task();
+		});
+	}
+	else
+	{
+		CallAbortTasksWithValidation();
+	}
+
+	if (IsNested())
+	{
 		AUTORTFM_ASSERT(Parent);
-    }
-    else
-    {
+	}
+	else
+	{
 		AUTORTFM_ASSERT(Context->IsAborting());
-    }
+	}
+}
+
+
+void FTransaction::ResetCommitTasksWithValidation()
+{
+	FOpenHashThrottler& OpenHashThrottler = Context->GetOpenHashThrottler();
+	const bool bThrottlingEnabled = AutoRTFM::ForTheRuntime::GetMemoryValidationThrottlingEnabled();
+	std::optional<FWriteHash> OldHash;
+	CommitTasks.RemoveEachBackward([&](TTask<void()>& Task)
+	{
+		void* const FunctionAddress = Task.FunctionAddress();
+
+		const bool bShouldHash = !bThrottlingEnabled || OpenHashThrottler.ShouldHashFor(FunctionAddress);
+		if (bShouldHash && !OldHash)
+		{
+			OldHash = CalculateNestedWriteHash(FunctionAddress);
+		}
+
+		Task.Reset();
+
+		if (bShouldHash)
+		{
+			const FWriteHash NewHash = CalculateNestedWriteHash(FunctionAddress);
+			if (OldHash != NewHash)
+			{
+				OnValidatorWriteHashMismatch(ForTheRuntime::GetMemoryValidationLevel(), 
+					"Memory modified in a transaction was also modified in the destructor of the on-commit handler: %s\n"
+					"This may lead to memory corruption if the transaction is aborted.",
+					GetFunctionDescription(FunctionAddress).c_str());
+				OldHash = NewHash;
+			}
+		}
+		else
+		{
+			// Task was reset without validation, so invalidate the snapshot hash.
+			OldHash = std::nullopt;
+		}
+	});
+}
+
+void FTransaction::CallAbortTasksWithValidation()
+{
+	FOpenHashThrottler& OpenHashThrottler = Context->GetOpenHashThrottler();
+	const bool bThrottlingEnabled = AutoRTFM::ForTheRuntime::GetMemoryValidationThrottlingEnabled();
+	std::optional<FWriteHash> OldHash;
+	auto MaybeValidate = [&](void* FunctionAddress, const char* Message, auto&& Fn)
+	{
+		const bool bShouldHash = !bThrottlingEnabled || OpenHashThrottler.ShouldHashFor(FunctionAddress);
+		if (bShouldHash && !OldHash)
+		{
+			OldHash = CalculateNestedWriteHash(FunctionAddress);
+		}
+		
+		Fn();
+
+		if (bShouldHash)
+		{
+			const FWriteHash NewHash = CalculateNestedWriteHash(FunctionAddress);
+			if (OldHash != NewHash)
+			{
+				OnValidatorWriteHashMismatch(ForTheRuntime::GetMemoryValidationLevel(), Message,
+					GetFunctionDescription(FunctionAddress).c_str());
+				OldHash = NewHash;
+			}
+		}
+		else
+		{
+			// Fn() was run without validation, so invalidate the snapshot hash.
+			OldHash = std::nullopt;
+		}
+	};
+
+	AbortTasks.RemoveEachBackward([&](TTask<void()>& Task)
+	{
+		void* const FunctionAddress = Task.FunctionAddress();
+
+		MaybeValidate(FunctionAddress,
+			"Memory modified in a transaction was also modified in the on-abort handler: %s\n"
+			"This may lead to memory corruption.",
+			[&]{ Task(); });
+
+		MaybeValidate(FunctionAddress,
+			"Memory modified in a transaction was also modified in the destructor of the on-abort handler: %s\n"
+			"This may lead to memory corruption if the transaction is aborted.",
+			[&]{ Task.Reset(); });
+	});
 }
 
 void FTransaction::AbortAndThrow()
@@ -441,7 +548,7 @@ void FTransaction::DebugBreakIfMemoryValidationFails()
 	if (CurrentMemoryValidationLevel != EMemoryValidationLevel::Disabled)
 	{
 		FWriteHash OldHash = RecordedWriteHash;
-		FWriteHash NewHash = CalculateNestedWriteHash();
+		FWriteHash NewHash = CalculateNestedWriteHash(CurrentOpenReturnAddress);
 		if (OldHash != NewHash)
 		{
 			AUTORTFM_WARN("DebugBreakIfInvalidMemoryHash() detected a change in hash");
@@ -453,53 +560,64 @@ void FTransaction::DebugBreakIfMemoryValidationFails()
 void FTransaction::RecordWriteHash()
 {
 	NumWriteLogsHashed = WriteLog.Num();
-	RecordedWriteHash = CalculateNestedWriteHash();
-	Context->GetOpenHashThrottler().Update();
+	RecordedWriteHash = CalculateNestedWriteHash(CurrentOpenReturnAddress);
 }
 
 void FTransaction::ValidateWriteHash() const
 {
-	FWriteHash OldHash = RecordedWriteHash;
-	FWriteHash NewHash = CalculateNestedWriteHash();
-	Context->GetOpenHashThrottler().Update();
+	const FWriteHash OldHash = RecordedWriteHash;
+	const FWriteHash NewHash = CalculateNestedWriteHash(CurrentOpenReturnAddress);
 
-	static constexpr const char Message[] =
-		"Memory modified in a transaction was also modified in an call to AutoRTFM::Open(). "
-		"This may lead to memory corruption if the transaction is aborted.";
-	if (AUTORTFM_UNLIKELY(OldHash != NewHash))
+	if (OldHash != NewHash)
 	{
-		if (CurrentMemoryValidationLevel == EMemoryValidationLevel::Warn)
-		{
-			AUTORTFM_WARN(Message);
-		}
-		else
-		{
+		OnValidatorWriteHashMismatch(CurrentMemoryValidationLevel,
+			"Memory modified in a transaction was also modified in an call to AutoRTFM::Open().\n"
+			"This may lead to memory corruption if the transaction is aborted.");
+	}
+}
+
+void FTransaction::OnValidatorWriteHashMismatch(EMemoryValidationLevel MemoryValidationLevel, const char* Message, ...) const
+{
+	va_list Args;
+	va_start(Args, Message);
+	switch (MemoryValidationLevel)
+	{
+		case EMemoryValidationLevel::Default:
+		case EMemoryValidationLevel::Disabled:
+			break;
+		case EMemoryValidationLevel::Warn:
+			AUTORTFM_WARN_V(Message, Args);
+			break;
+		case EMemoryValidationLevel::Error:
 			if (!ForTheRuntime::GetEnsureOnInternalAbort())
 			{
-				AUTORTFM_FATAL(Message);
+				AUTORTFM_FATAL_V(Message, Args);
 			}
 			else
 			{
-				AUTORTFM_ENSURE_MSG(OldHash == NewHash, Message);
+				AUTORTFM_ENSURE_MSG_V(false, Message, Args);
 			}
-		}
+			break;
 	}
+	va_end(Args);
 }
 
-FTransaction::FWriteHash FTransaction::CalculateNestedWriteHash() const
+FTransaction::FWriteHash FTransaction::CalculateNestedWriteHash(const void *FunctionAddress) const
 {
-	return CalculateNestedWriteHashWithLimit(NumWriteLogsHashed, CurrentOpenReturnAddress);
+	const FWriteHash Hash = CalculateNestedWriteHashWithLimit(NumWriteLogsHashed, FunctionAddress);
+	Context->GetOpenHashThrottler().Update();
+	return Hash;
 }
 
-FTransaction::FWriteHash FTransaction::CalculateNestedWriteHashWithLimit(size_t NumWriteEntries, const void *OpenReturnAddress) const
+FTransaction::FWriteHash FTransaction::CalculateNestedWriteHashWithLimit(size_t NumWriteEntries, const void *FunctionAddress) const
 {
 	FWriteHash Hash = 0;
 	if (nullptr != Parent)
 	{
-		Hash = 31 * Parent->CalculateNestedWriteHashWithLimit(Parent->WriteLog.Num(), OpenReturnAddress);
+		Hash = 31 * Parent->CalculateNestedWriteHashWithLimit(Parent->WriteLog.Num(), FunctionAddress);
 	}
 	{
-		FOpenHashThrottler::FHashScope Profile(Context->GetOpenHashThrottler(), OpenReturnAddress, WriteLog);
+		FOpenHashThrottler::FHashScope Profile(Context->GetOpenHashThrottler(), FunctionAddress, WriteLog);
 		Hash ^= WriteLog.Hash(NumWriteEntries);
 	}
 	return Hash;
diff --git a/Engine/Source/Runtime/AutoRTFM/Private/Transaction.h b/Engine/Source/Runtime/AutoRTFM/Private/Transaction.h
--- a/Engine/Source/Runtime/AutoRTFM/Private/Transaction.h
+++ b/Engine/Source/Runtime/AutoRTFM/Private/Transaction.h
@@ -162,6 +162,9 @@ private:
 	void CommitNested();
 	bool AttemptToCommitOuterNest();
 
+	void ResetCommitTasksWithValidation();
+	void CallAbortTasksWithValidation();
+
 	// Calculates the hash of all the memory written during the transaction and stores it to
 	// RecordedWriteHash.
 	void RecordWriteHash();
@@ -171,15 +174,19 @@ private:
 	// See CalculateNestedWriteHash().
 	void ValidateWriteHash() const;
 
+	// Warns / errors when a validation hash mismatch is detected.
+	void OnValidatorWriteHashMismatch(EMemoryValidationLevel MemoryValidationLevel, const char* Message, ...) const;
+	
 	// Returns a hash of all the memory written during this transaction and all parent transactions.
 	// Can be used to ensure that no transaction written memory is modified during an open.
-	FWriteHash CalculateNestedWriteHash() const;
+	// FunctionAddress is the address of the function that triggered this hash.
+	FWriteHash CalculateNestedWriteHash(const void *FunctionAddress) const;
 	
 	// Returns a hash of all the memory written during this transaction and all parent transactions.
 	// Can be used to ensure that no transaction written memory is modified during an open.
 	// NumWriteEntries is the number of write entries to hash before stopping.
-	// OpenReturnAddress is the return address for the open that triggered this hash.
-	FWriteHash CalculateNestedWriteHashWithLimit(size_t NumWriteEntries, const void *OpenReturnAddress) const;
+	// FunctionAddress is the address of the function that triggered this hash.
+	FWriteHash CalculateNestedWriteHashWithLimit(size_t NumWriteEntries, const void *FunctionAddress) const;
 
 	void CollectStats() const;
 
diff --git a/Engine/Source/Runtime/AutoRTFM/Public/AutoRTFMTask.h b/Engine/Source/Runtime/AutoRTFM/Public/AutoRTFMTask.h
--- a/Engine/Source/Runtime/AutoRTFM/Public/AutoRTFMTask.h
+++ b/Engine/Source/Runtime/AutoRTFM/Public/AutoRTFMTask.h
@@ -234,6 +234,12 @@ public:
 		return PayloadMethods != nullptr;
 	}
 
+	// Returns the address of the functor.
+	void* FunctionAddress() const
+	{
+		return reinterpret_cast<void*>(PayloadMethods->Call);
+	}
+
 private:
 	// mutable as the functor might be mutable.
 	alignas(InlineDataAlignment) mutable FPayload Payload;
