[AutoRTFM] Enforce that AUTORTFM_DISABLED is on the declaration & support AUTORTFM_DISABLED on classes.

Enforcing that AUTORTFM_DISABLED is only applied to a function declaration ensures that callers have visibility on the attribute, which is required to prevent calls from the closed at compile time (not enforced yet).
Applying AUTORTFM_DISABLED to a class is equivalent to applying the attribute to all methods of the class.

diff --git a/Engine/Source/Runtime/AutoRTFM/Public/AutoRTFMDefines.h b/Engine/Source/Runtime/AutoRTFM/Public/AutoRTFMDefines.h
--- a/Engine/Source/Runtime/AutoRTFM/Public/AutoRTFMDefines.h
+++ b/Engine/Source/Runtime/AutoRTFM/Public/AutoRTFMDefines.h
@@ -37,9 +37,13 @@
 // call.
 #define UE_AUTORTFM_ALWAYS_OPEN_NO_MEMORY_VALIDATION [[clang::autortfm_always_open_disable_memory_validation]]
 
-// The annotated function will have no AutoRTFM closed variant generated, and
-// cannot be called from another closed function. Must be placed on the function
-// declaration (usually in the header), not the implementation.
+// Annotation that can be applied to classes, methods or functions to prevent
+// AutoRTFM closed function(s) from being generated. Applying the annotation to
+// a class is equivalent to applying the annotation to each method of the class.
+// When annotating a function, the attribute must be placed on a the function
+// declaration (usually in the header) and not a function implementation.
+// Annotated functions cannot be called from another closed function and
+// attempting to do so will result in a runtime failure.
 #define AUTORTFM_DISABLE [[clang::autortfm_disable]]
 
 // Force the call statement to be inlined.
