[AutoRTFM] Error if calling AUTORTFM_DISABLE annotated functions from AutoRTFM-enabled functions

Provides some useful information on how you might fix this, if seen.
Note: The compiler will only detect invalid calls post-optimization passes, so inlining may prevent errors from being raised for violations. This will be improved in follow-up changes.

Also:
* Enforce that AUTORTFM_CORE_API is only applied to the function declaration, not the definition, similarly to AUTORTFM_DISABLE.
* Allow AUTORTFM_CORE_API to be applied to a class, similarly to AUTORTFM_DISABLE.
* Apply AUTORTFM_CORE_API to AutoRTFM::Context and AutoRTFM::Transaction. There are methods called from both the open (rw) and closed (r), but we don't want instrumentation on these.
* Rework autortfm_unreachable() to be fatal and [[noreturn]]. This allows the compiler to dead-code eliminate code dominated by the call, and to be used as a way to silence the new AUTORTFM_DISABLE call enforcement.
* Add AutoRTFM::UnreachableIfClosed()
* Add autortfm_disable_called(), and have the compiler call this when a AUTORTFM_DISABLE annotated and inlined function is called in the closed. Fixes an incorrect and confusing error message.
* Add a thorough set of lit tests for [[clang::autortfm_always_open]], [[clang::autortfm_core_api]] and [[clang::autortfm_disable]].
* Add a missing unit test for `autortfm_unreachable()`

diff --git a/Engine/Source/Programs/AutoRTFMTests/Private/API.cpp b/Engine/Source/Programs/AutoRTFMTests/Private/API.cpp
--- a/Engine/Source/Programs/AutoRTFMTests/Private/API.cpp
+++ b/Engine/Source/Programs/AutoRTFMTests/Private/API.cpp
@@ -2,12 +2,11 @@
 
 #include "API.h"
 
+#include "AutoRTFM.h"
+#include "AutoRTFMTestUtils.h"
 #include "AutoRTFMTesting.h"
 #include "Catch2Includes.h"
 
-#include <memory>
-#include <thread>
-
 TEST_CASE("API.autortfm_is_transactional")
 {
     REQUIRE(false == autortfm_is_transactional());
@@ -333,6 +332,23 @@ TEST_CASE("API.autortfm_did_allocate")
     }
 }
 
+TEST_CASE("API.autortfm_unreachable")
+{
+	AutoRTFMTestUtils::FScopedInternalAbortAction Scoped1(AutoRTFM::ForTheRuntime::EAutoRTFMInternalAbortActionState::Abort);
+	AutoRTFMTestUtils::FScopedEnsureOnInternalAbort Scoped2(false);
+
+	bool bTouched = false;
+
+	const AutoRTFM::ETransactionResult Result = AutoRTFM::Transact([&]
+	{
+		bTouched = true;
+		autortfm_unreachable("Oh noes!");
+	});
+
+	REQUIRE(AutoRTFM::ETransactionResult::AbortedByLanguage == Result);
+	REQUIRE(false == bTouched);
+}
+
 TEST_CASE("API.ETransactionResult")
 {
     int Answer = 6 * 9;
diff --git a/Engine/Source/Programs/AutoRTFMTests/Private/TaskTests.cpp b/Engine/Source/Programs/AutoRTFMTests/Private/TaskTests.cpp
--- a/Engine/Source/Programs/AutoRTFMTests/Private/TaskTests.cpp
+++ b/Engine/Source/Programs/AutoRTFMTests/Private/TaskTests.cpp
@@ -1,11 +1,13 @@
 // Copyright Epic Games, Inc. All Rights Reserved.
 
+#include "AutoRTFMDefines.h"
 #include "AutoRTFMTask.h"
 #include "Catch2Includes.h"
 #include "GenericPlatform/GenericPlatformMisc.h"
 #include "Templates/UnrealTemplate.h"
 #include <functional>
 
+AUTORTFM_DISABLE // Calls Task::operator() which is AUTORTFM_DISABLE'd
 TEST_CASE("TTask")
 {
 	SECTION("void()")
diff --git a/Engine/Source/Runtime/AutoRTFM/Private/API.cpp b/Engine/Source/Runtime/AutoRTFM/Private/API.cpp
--- a/Engine/Source/Runtime/AutoRTFM/Private/API.cpp
+++ b/Engine/Source/Runtime/AutoRTFM/Private/API.cpp
@@ -422,8 +422,6 @@ namespace AutoRTFM
 // This is the implementation of the AutoRTFM.h API. Ideally, functions here should just delegate to some internal API.
 // For now, I have these functions also perform some error checking.
 
-#define AUTORTFM_CORE_API [[clang::autortfm_core_api]]
-
 namespace AutoRTFM
 {
 
diff --git a/Engine/Source/Runtime/AutoRTFM/Private/CompilerInterface.cpp b/Engine/Source/Runtime/AutoRTFM/Private/CompilerInterface.cpp
--- a/Engine/Source/Runtime/AutoRTFM/Private/CompilerInterface.cpp
+++ b/Engine/Source/Runtime/AutoRTFM/Private/CompilerInterface.cpp
@@ -139,6 +139,7 @@ extern "C" UE_AUTORTFM_API void autortfm_memset(void* Dst, int Value, size_t Siz
 extern "C" UE_AUTORTFM_API void autortfm_unreachable(const char* Message) noexcept
 {
 	AUTORTFM_REPORT_ERROR("AutoRTFM Unreachable: %s", Message);
+	__builtin_unreachable();
 }
 
 extern "C" UE_AUTORTFM_API void autortfm_llvm_fail(const char* Message) noexcept
@@ -173,6 +174,10 @@ extern "C" UE_AUTORTFM_API void autortfm_called_no_autortfm() noexcept
 	AUTORTFM_FATAL("inlined UE_AUTORTFM_NOAUTORTFM function called from the closed");
 }
 
+extern "C" UE_AUTORTFM_API void autortfm_disable_called() noexcept
+{
+	AUTORTFM_FATAL("inlined AUTORTFM_DISABLE function called from the closed");
+}
 } // namespace AutoRTFM
 
 #endif // defined(__AUTORTFM) && __AUTORTFM
diff --git a/Engine/Source/Runtime/AutoRTFM/Private/Context.h b/Engine/Source/Runtime/AutoRTFM/Private/Context.h
--- a/Engine/Source/Runtime/AutoRTFM/Private/Context.h
+++ b/Engine/Source/Runtime/AutoRTFM/Private/Context.h
@@ -10,6 +10,7 @@
 #include "TaskArray.h"
 #include "ThreadID.h"
 #include "Transaction.h"
+#include "Utils.h"
 
 namespace AutoRTFM
 {
@@ -17,7 +18,7 @@ namespace AutoRTFM
 class FLineLock;
 class FCallNest;
 
-class FContext final
+class AUTORTFM_CORE_API FContext final
 {
 	using FTaskPool = typename TTaskArray<TTask<void()>>::FEntryPool;
 	using FTransactionPool = TIntrusivePool<FTransaction, 16>;
@@ -41,8 +42,8 @@ public:
 	ETransactionResult RollbackTransaction(EContextStatus NewStatus);
 
 	ETransactionResult CommitTransaction();
-	void AbortTransaction(EContextStatus NewStatus);
-	void AbortTransactionWithPostAbortCallback(EContextStatus NewStatus, TTask<void()>&& InCallback);
+	[[noreturn]] void AbortTransaction(EContextStatus NewStatus);
+	[[noreturn]] void AbortTransactionWithPostAbortCallback(EContextStatus NewStatus, TTask<void()>&& InCallback);
 	void ClearTransactionStatus();
 	bool IsAborting() const;
 	bool IsTransactional() const;
diff --git a/Engine/Source/Runtime/AutoRTFM/Private/Transaction.h b/Engine/Source/Runtime/AutoRTFM/Private/Transaction.h
--- a/Engine/Source/Runtime/AutoRTFM/Private/Transaction.h
+++ b/Engine/Source/Runtime/AutoRTFM/Private/Transaction.h
@@ -12,13 +12,14 @@
 #include "StackRange.h"
 #include "Stats.h"
 #include "TaskArray.h"
+#include "Utils.h"
 #include "WriteLog.h"
 
 namespace AutoRTFM
 {
 class FContext;
 
-class FTransaction final
+class AUTORTFM_CORE_API FTransaction final
 {
 	friend class TIntrusivePool<FTransaction, 16>;
 
diff --git a/Engine/Source/Runtime/AutoRTFM/Private/Utils.h b/Engine/Source/Runtime/AutoRTFM/Private/Utils.h
--- a/Engine/Source/Runtime/AutoRTFM/Private/Utils.h
+++ b/Engine/Source/Runtime/AutoRTFM/Private/Utils.h
@@ -7,15 +7,23 @@
 #include "AutoRTFM.h"
 #include "ExternAPI.h"
 
-#include <algorithm>
 #include <cstdarg>
 #include <string>
 #include <type_traits>
-#include <utility>
 
 // If 0 then verbose logging is compiled-out.
 #define AUTORTFM_VERBOSE_ENABLED 0
 
+// AutoRTFM internal-runtime-only attribute that can be applied to classes, 
+// methods or functions to prevent AutoRTFM closed function(s) from being
+// generated. Applying the annotation to a class is equivalent to applying the
+// annotation to each method of the class.
+// When annotating a function, the attribute must be placed on a the function
+// declaration (usually in the header) and not a function implementation.
+// Annotated functions are callable from both the open and closed, without
+// trampolines or validation, and the function is always uninstrumented.
+#define AUTORTFM_CORE_API [[clang::autortfm_core_api]]
+
 namespace AutoRTFM
 {
 
diff --git a/Engine/Source/Runtime/AutoRTFM/Public/AutoRTFM.h b/Engine/Source/Runtime/AutoRTFM/Public/AutoRTFM.h
--- a/Engine/Source/Runtime/AutoRTFM/Public/AutoRTFM.h
+++ b/Engine/Source/Runtime/AutoRTFM/Public/AutoRTFM.h
@@ -421,7 +421,10 @@ UE_AUTORTFM_CRITICAL_INLINE void autortfm_check_abi(void* ptr, size_t size) AUTO
 #endif
 
 #if UE_AUTORTFM_ENABLED
-UE_AUTORTFM_API void autortfm_unreachable(const char* Message) AUTORTFM_NOEXCEPT;
+// Called when execution unexpectedly reaches a code path that was considered unreachable.
+// Either aborts execution of the program or aborts the current transaction, depending on the
+// current 'InternalAbortAction' state.
+[[noreturn]] UE_AUTORTFM_API void autortfm_unreachable(const char* Message) AUTORTFM_NOEXCEPT;
 #else
 UE_AUTORTFM_CRITICAL_INLINE void autortfm_unreachable(const char* Message) AUTORTFM_NOEXCEPT
 {
@@ -1071,10 +1074,17 @@ template<typename Type> UE_AUTORTFM_CRITICAL_INLINE void RecordOpenWriteNoMemory
 
 // Report that a unreachable codepath is being hit. Used to manually ban certain codepaths
 // from being transactionally safe.
-UE_AUTORTFM_CRITICAL_INLINE void Unreachable(const char* Message = nullptr)
+#if UE_AUTORTFM_ENABLED
+[[noreturn]] UE_AUTORTFM_CRITICAL_INLINE void Unreachable(const char* Message = nullptr)
 {
 	autortfm_unreachable(Message);
 }
+#else
+UE_AUTORTFM_CRITICAL_INLINE void Unreachable(const char* Message = nullptr)
+{
+	UE_AUTORTFM_UNUSED(Message);
+}
+#endif
 
 // If we are running within a transaction, call `AutoRTFM::Unreachable`.
 UE_AUTORTFM_CRITICAL_INLINE void UnreachableIfTransactional(const char* Message = nullptr)
@@ -1085,6 +1095,15 @@ UE_AUTORTFM_CRITICAL_INLINE void UnreachableIfTransactional(const char* Message
 	}
 }
 
+// If we are running within a closed transaction, call `AutoRTFM::Unreachable`.
+UE_AUTORTFM_CRITICAL_INLINE void UnreachableIfClosed(const char* Message = nullptr)
+{
+	if (AutoRTFM::IsClosed())
+	{
+		AutoRTFM::Unreachable(Message);
+	}
+}
+
 // A collection of power-user functions that are reserved for use by the AutoRTFM runtime only.
 namespace ForTheRuntime
 {
