[AutoRTFM] Annotate the AutoRTFM runtime with AUTORTFM_DISABLE, rename AUTORTFM_CORE_API -> AUTORTFM_INTERNAL

Rename AUTORTFM_CORE_API to AUTORTFM_INTERNAL, as dllexport / dllimports use the `_API` suffix and `_CORE` looks like it has something to do with UE's Core, which it doesn't.
The underlying compiler attribute will be fixed in a later CL.

Replace the use of AUTORTFM_CORE_API on runtime classes with AUTORTFM_DISABLE. This serves two purposes:
* The compiler will now error if accidentally called from the closed.
* There was an inconsistency with AUTORTFM_CORE_API where direct annotations on functions would prevent inlining, but on classes it wouldn't. This will be fixed in a future compiler change, but we want these class methods to be inlined. AUTORTFM_DISABLE allows for inlining.

Callers to the AUTORTFM_DISABLE annotated functions have been either decorated with the appropriate annotation to disable AutoRTFM, always-open, prefix the call with AutoRTFM::UnreachableIfClosed(), or (in the case of tests) special case the calls.

UE_AUTORTFM_FORCEINLINE_ALWAYS and UE_AUTORTFM_CRITICAL_INLINE_ALWAYS have been added to force-inline even in debug builds. These have been used on select function to ensure that AutoRTFM::UnreachableIfClosed() can be used to silence compiler errors for AUTORTFM_DISABLED function calls, even in debug builds.

diff --git a/Engine/Source/Programs/AutoRTFMTests/Private/Statics.cpp b/Engine/Source/Programs/AutoRTFMTests/Private/Statics.cpp
--- a/Engine/Source/Programs/AutoRTFMTests/Private/Statics.cpp
+++ b/Engine/Source/Programs/AutoRTFMTests/Private/Statics.cpp
@@ -1,12 +1,14 @@
 // Copyright Epic Games, Inc. All Rights Reserved.
 
 #include "AutoRTFM.h"
+#include "AutoRTFMDefines.h"
 #include "AutoRTFMTesting.h"
 #include "Context.h"
 #include "ContextInlines.h"
 #include "Transaction.h"
 #include "Catch2Includes.h"
 #include "CrossCUTests.h"
+#include "Utils.h"
 
 #include <chrono>
 #include <condition_variable>
@@ -306,6 +308,15 @@ TEST_CASE("Static.IsTransactional")
 // Test case for SOL-7360
 TEST_CASE("Static.Concurrent")
 {
+	// Helper for inspecting the current transaction's state without changing it.
+	struct FTransactionState
+	{
+		AUTORTFM_INTERNAL static AutoRTFM::FTransaction::EState Get()
+		{
+			return AutoRTFM::FContext::Get()->GetCurrentTransaction()->State();
+		}
+	};
+	
 	// A simple Signal() / Wait() event.
 	class FEvent
 	{
@@ -357,10 +368,11 @@ TEST_CASE("Static.Concurrent")
 	// A separate thread used to execute a static local initializer under a transaction.
 	std::thread TransactionalThread = std::thread([&]
 	{
+		AutoRTFM::UnreachableIfTransactional();
 		AutoRTFM::Transact([&]
 		{
-			AutoRTFM::FTransaction* Transaction = AutoRTFM::FContext::Get()->GetCurrentTransaction();
-			REQUIRE(Transaction->State() == AutoRTFM::FTransaction::EState::ClosedActive);
+			REQUIRE(FTransactionState::Get() == AutoRTFM::FTransaction::EState::ClosedActive);
+			
 			// Signal to the main thread that this thread is about to wait on
 			// UnblockTransactionThread. This attempts to keep the threads
 			// tightly synchronized.
@@ -378,7 +390,7 @@ TEST_CASE("Static.Concurrent")
 
 			// If SOL-7360 is not fixed, then the transaction state will be incorrect.
 			// (An assert would have likely been hit already too)
-			REQUIRE(Transaction->State() == AutoRTFM::FTransaction::EState::ClosedActive);
+			REQUIRE(FTransactionState::Get() == AutoRTFM::FTransaction::EState::ClosedActive);
 		});
 	});
 
diff --git a/Engine/Source/Programs/AutoRTFMTests/Private/TaskTests.cpp b/Engine/Source/Programs/AutoRTFMTests/Private/TaskTests.cpp
--- a/Engine/Source/Programs/AutoRTFMTests/Private/TaskTests.cpp
+++ b/Engine/Source/Programs/AutoRTFMTests/Private/TaskTests.cpp
@@ -1,13 +1,26 @@
 // Copyright Epic Games, Inc. All Rights Reserved.
 
-#include "AutoRTFMDefines.h"
+#include "AutoRTFM.h"
 #include "AutoRTFMTask.h"
 #include "Catch2Includes.h"
 #include "GenericPlatform/GenericPlatformMisc.h"
 #include "Templates/UnrealTemplate.h"
 #include <functional>
 
-AUTORTFM_DISABLE // Calls Task::operator() which is AUTORTFM_DISABLE'd
+namespace
+{
+
+// Helper for calling Task with the given arguments, which may be a
+// AUTORTFM_DISABLED annotated function.
+template<typename TASK, typename ... ARGS>
+auto CallTask(TASK&& Task, ARGS&& ... Args)
+{
+	AutoRTFM::UnreachableIfClosed();
+	return Task(std::forward<ARGS>(Args)...);
+}
+
+} // anonymous namespace
+
 TEST_CASE("TTask")
 {
 	SECTION("void()")
@@ -83,7 +96,7 @@ TEST_CASE("TTask")
 				FTask Original(SmallLambda);
 				FTask Task(Original);
 				Original.Reset();
-				Task();
+				CallTask(Task);
 				REQUIRE(WasSmallLambdaCalled());
 			}
 			SECTION("Large")
@@ -91,7 +104,7 @@ TEST_CASE("TTask")
 				FTask Original(LargeLambda);
 				FTask Task(Original);
 				Original.Reset();
-				Task();
+				CallTask(Task);
 				REQUIRE(WasLargeLambdaCalled());
 			}
 			SECTION("Aligned")
@@ -99,7 +112,7 @@ TEST_CASE("TTask")
 				FTask Original(AlignedLambda);
 				FTask Task(Original);
 				Original.Reset();
-				Task();
+				CallTask(Task);
 				REQUIRE(WasAlignedLambdaCalled());
 			}
 		}
@@ -109,19 +122,19 @@ TEST_CASE("TTask")
 			SECTION("Small")
 			{
 				FTask Task(SmallLambda);
-				Task();
+				CallTask(Task);
 				REQUIRE(WasSmallLambdaCalled());
 			}
 			SECTION("Large")
 			{
 				FTask Task(LargeLambda);
-				Task();
+				CallTask(Task);
 				REQUIRE(WasLargeLambdaCalled());
 			}
 			SECTION("Aligned")
 			{
 				FTask Task(AlignedLambda);
-				Task();
+				CallTask(Task);
 				REQUIRE(WasAlignedLambdaCalled());
 			}
 		}
@@ -131,7 +144,7 @@ TEST_CASE("TTask")
 			TFunction<void()> Function(SmallLambda);
 			FTask Task(Function);
 			REQUIRE(Function.IsSet() == true);
-			Task();
+			CallTask(Task);
 			REQUIRE(WasSmallLambdaCalled());
 		}
 
@@ -140,7 +153,7 @@ TEST_CASE("TTask")
 			std::function<void()> Function(SmallLambda);
 			FTask Task(Function);
 			REQUIRE(static_cast<bool>(Function) == true);
-			Task();
+			CallTask(Task);
 			REQUIRE(WasSmallLambdaCalled());
 		}
 
@@ -159,7 +172,7 @@ TEST_CASE("TTask")
 				FTask Task;
 				Task = Original;
 				Original.Reset();
-				Task();
+				CallTask(Task);
 				REQUIRE(WasSmallLambdaCalled());
 			}
 			SECTION("Large")
@@ -168,7 +181,7 @@ TEST_CASE("TTask")
 				FTask Task;
 				Task = Original;
 				Original.Reset();
-				Task();
+				CallTask(Task);
 				REQUIRE(WasLargeLambdaCalled());
 			}
 			SECTION("Aligned")
@@ -177,7 +190,7 @@ TEST_CASE("TTask")
 				FTask Task;
 				Task = Original;
 				Original.Reset();
-				Task();
+				CallTask(Task);
 				REQUIRE(WasAlignedLambdaCalled());
 			}
 		}
@@ -188,21 +201,21 @@ TEST_CASE("TTask")
 			{
 				FTask Task;
 				Task = SmallLambda;
-				Task();
+				CallTask(Task);
 				REQUIRE(WasSmallLambdaCalled());
 			}
 			SECTION("Large")
 			{
 				FTask Task;
 				Task = LargeLambda;
-				Task();
+				CallTask(Task);
 				REQUIRE(WasLargeLambdaCalled());
 			}
 			SECTION("Aligned")
 			{
 				FTask Task;
 				Task = AlignedLambda;
-				Task();
+				CallTask(Task);
 				REQUIRE(WasAlignedLambdaCalled());
 			}
 		}
@@ -213,7 +226,7 @@ TEST_CASE("TTask")
 			FTask Task;
 			Task = Function;
 			REQUIRE(Function.IsSet() == true);
-			Task();
+			CallTask(Task);
 			REQUIRE(WasSmallLambdaCalled());
 		}
 
@@ -223,7 +236,7 @@ TEST_CASE("TTask")
 			FTask Task;
 			Task = Function;
 			REQUIRE(static_cast<bool>(Function) == true);
-			Task();
+			CallTask(Task);
 			REQUIRE(WasSmallLambdaCalled());
 		}
 
@@ -232,19 +245,19 @@ TEST_CASE("TTask")
 			SECTION("Small")
 			{
 				FTask Task(MoveTemp(SmallLambda));
-				Task();
+				CallTask(Task);
 				REQUIRE(WasSmallLambdaCalled());
 			}
 			SECTION("Large")
 			{
 				FTask Task(MoveTemp(LargeLambda));
-				Task();
+				CallTask(Task);
 				REQUIRE(WasLargeLambdaCalled());
 			}
 			SECTION("Aligned")
 			{
 				FTask Task(MoveTemp(AlignedLambda));
-				Task();
+				CallTask(Task);
 				REQUIRE(WasAlignedLambdaCalled());
 			}
 		}
@@ -256,7 +269,7 @@ TEST_CASE("TTask")
 				FTask Original(SmallLambda);
 				FTask Task(MoveTemp(Original));
 				REQUIRE(Original.IsSet() == false);
-				Task();
+				CallTask(Task);
 				REQUIRE(WasSmallLambdaCalled());
 			}
 			SECTION("Large")
@@ -264,7 +277,7 @@ TEST_CASE("TTask")
 				FTask Original(LargeLambda);
 				FTask Task(MoveTemp(Original));
 				REQUIRE(Original.IsSet() == false);
-				Task();
+				CallTask(Task);
 				REQUIRE(WasLargeLambdaCalled());
 			}
 			SECTION("Aligned")
@@ -272,7 +285,7 @@ TEST_CASE("TTask")
 				FTask Original(AlignedLambda);
 				FTask Task(MoveTemp(Original));
 				REQUIRE(Original.IsSet() == false);
-				Task();
+				CallTask(Task);
 				REQUIRE(WasAlignedLambdaCalled());
 			}
 		}
@@ -282,7 +295,7 @@ TEST_CASE("TTask")
 			TFunction<void()> Function(SmallLambda);
 			FTask Task(MoveTemp(Function));
 			REQUIRE(Function.IsSet() == false);
-			Task();
+			CallTask(Task);
 			REQUIRE(WasSmallLambdaCalled());
 		}
 
@@ -292,7 +305,7 @@ TEST_CASE("TTask")
 			FTask Task(MoveTemp(Function));
 			// Note: standard library implementations of std::function are not
 			// guaranteed to invalidate the function on move - hence no check.
-			Task();
+			CallTask(Task);
 			REQUIRE(WasSmallLambdaCalled());
 		}
 
@@ -302,21 +315,21 @@ TEST_CASE("TTask")
 			{
 				FTask Task;
 				Task = MoveTemp(SmallLambda);
-				Task();
+				CallTask(Task);
 				REQUIRE(WasSmallLambdaCalled());
 			}
 			SECTION("Large")
 			{
 				FTask Task;
 				Task = MoveTemp(LargeLambda);
-				Task();
+				CallTask(Task);
 				REQUIRE(WasLargeLambdaCalled());
 			}
 			SECTION("Aligned")
 			{
 				FTask Task;
 				Task = MoveTemp(AlignedLambda);
-				Task();
+				CallTask(Task);
 				REQUIRE(WasAlignedLambdaCalled());
 			}
 		}
@@ -329,7 +342,7 @@ TEST_CASE("TTask")
 				FTask Task;
 				Task = MoveTemp(Original);
 				REQUIRE(Original.IsSet() == false);
-				Task();
+				CallTask(Task);
 				REQUIRE(WasSmallLambdaCalled());
 			}
 			SECTION("Large")
@@ -338,7 +351,7 @@ TEST_CASE("TTask")
 				FTask Task;
 				Task = MoveTemp(Original);
 				REQUIRE(Original.IsSet() == false);
-				Task();
+				CallTask(Task);
 				REQUIRE(WasLargeLambdaCalled());
 			}
 			SECTION("Aligned")
@@ -347,7 +360,7 @@ TEST_CASE("TTask")
 				FTask Task;
 				Task = MoveTemp(Original);
 				REQUIRE(Original.IsSet() == false);
-				Task();
+				CallTask(Task);
 				REQUIRE(WasAlignedLambdaCalled());
 			}
 		}
@@ -358,7 +371,7 @@ TEST_CASE("TTask")
 			FTask Task;
 			Task = MoveTemp(Function);
 			REQUIRE(Function.IsSet() == false);
-			Task();
+			CallTask(Task);
 			REQUIRE(WasSmallLambdaCalled());
 		}
 
@@ -369,7 +382,7 @@ TEST_CASE("TTask")
 			Task = MoveTemp(Function);
 			// Note: standard library implementations of std::function are not
 			// guaranteed to invalidate the function on move - hence no check.
-			Task();
+			CallTask(Task);
 			REQUIRE(WasSmallLambdaCalled());
 		}
 
@@ -470,25 +483,25 @@ TEST_CASE("TTask")
 			SECTION("Small")
 			{
 				FTask Task(SmallLambda);
-				REQUIRE(Task(10, true) == 10);
+				REQUIRE(CallTask(Task, 10, true) == 10);
 				REQUIRE(WasSmallLambdaCalled(10));
-				REQUIRE(Task(20, false) == 20);
+				REQUIRE(CallTask(Task, 20, false) == 20);
 				REQUIRE(WasSmallLambdaCalled(-20));
 			}
 			SECTION("Large")
 			{
 				FTask Task(LargeLambda);
-				REQUIRE(Task(10, true) == 10);
+				REQUIRE(CallTask(Task, 10, true) == 10);
 				REQUIRE(WasLargeLambdaCalled(10));
-				REQUIRE(Task(20, false) == 20);
+				REQUIRE(CallTask(Task, 20, false) == 20);
 				REQUIRE(WasLargeLambdaCalled(-20));
 			}
 			SECTION("Aligned")
 			{
 				FTask Task(AlignedLambda);
-				REQUIRE(Task(10, true) == 10);
+				REQUIRE(CallTask(Task, 10, true) == 10);
 				REQUIRE(WasAlignedLambdaCalled(10));
-				REQUIRE(Task(20, false) == 20);
+				REQUIRE(CallTask(Task, 20, false) == 20);
 				REQUIRE(WasAlignedLambdaCalled(-20));
 			}
 		}
@@ -506,9 +519,9 @@ TEST_CASE("TTask")
 				FTask Original(SmallLambda);
 				FTask Task(Original);
 				Original.Reset();
-				REQUIRE(Task(10, true) == 10);
+				REQUIRE(CallTask(Task, 10, true) == 10);
 				REQUIRE(WasSmallLambdaCalled(10));
-				REQUIRE(Task(20, false) == 20);
+				REQUIRE(CallTask(Task, 20, false) == 20);
 				REQUIRE(WasSmallLambdaCalled(-20));
 			}
 			SECTION("Large")
@@ -516,9 +529,9 @@ TEST_CASE("TTask")
 				FTask Original(LargeLambda);
 				FTask Task(Original);
 				Original.Reset();
-				REQUIRE(Task(10, true) == 10);
+				REQUIRE(CallTask(Task, 10, true) == 10);
 				REQUIRE(WasLargeLambdaCalled(10));
-				REQUIRE(Task(20, false) == 20);
+				REQUIRE(CallTask(Task, 20, false) == 20);
 				REQUIRE(WasLargeLambdaCalled(-20));
 			}
 			SECTION("Aligned")
@@ -526,9 +539,9 @@ TEST_CASE("TTask")
 				FTask Original(AlignedLambda);
 				FTask Task(Original);
 				Original.Reset();
-				REQUIRE(Task(10, true) == 10);
+				REQUIRE(CallTask(Task, 10, true) == 10);
 				REQUIRE(WasAlignedLambdaCalled(10));
-				REQUIRE(Task(20, false) == 20);
+				REQUIRE(CallTask(Task, 20, false) == 20);
 				REQUIRE(WasAlignedLambdaCalled(-20));
 			}
 		}
@@ -538,25 +551,25 @@ TEST_CASE("TTask")
 			SECTION("Small")
 			{
 				FTask Task(MoveTemp(SmallLambda));
-				REQUIRE(Task(10, true) == 10);
+				REQUIRE(CallTask(Task, 10, true) == 10);
 				REQUIRE(WasSmallLambdaCalled(10));
-				REQUIRE(Task(20, false) == 20);
+				REQUIRE(CallTask(Task, 20, false) == 20);
 				REQUIRE(WasSmallLambdaCalled(-20));
 			}
 			SECTION("Large")
 			{
 				FTask Task(MoveTemp(LargeLambda));
-				REQUIRE(Task(10, true) == 10);
+				REQUIRE(CallTask(Task, 10, true) == 10);
 				REQUIRE(WasLargeLambdaCalled(10));
-				REQUIRE(Task(20, false) == 20);
+				REQUIRE(CallTask(Task, 20, false) == 20);
 				REQUIRE(WasLargeLambdaCalled(-20));
 			}
 			SECTION("Aligned")
 			{
 				FTask Task(MoveTemp(AlignedLambda));
-				REQUIRE(Task(10, true) == 10);
+				REQUIRE(CallTask(Task, 10, true) == 10);
 				REQUIRE(WasAlignedLambdaCalled(10));
-				REQUIRE(Task(20, false) == 20);
+				REQUIRE(CallTask(Task, 20, false) == 20);
 				REQUIRE(WasAlignedLambdaCalled(-20));
 			}
 		}
@@ -568,9 +581,9 @@ TEST_CASE("TTask")
 				FTask Original(SmallLambda);
 				FTask Task(MoveTemp(Original));
 				REQUIRE(Original.IsSet() == false);
-				REQUIRE(Task(10, true) == 10);
+				REQUIRE(CallTask(Task, 10, true) == 10);
 				REQUIRE(WasSmallLambdaCalled(10));
-				REQUIRE(Task(20, false) == 20);
+				REQUIRE(CallTask(Task, 20, false) == 20);
 				REQUIRE(WasSmallLambdaCalled(-20));
 			}
 			SECTION("Large")
@@ -578,9 +591,9 @@ TEST_CASE("TTask")
 				FTask Original(LargeLambda);
 				FTask Task(MoveTemp(Original));
 				REQUIRE(Original.IsSet() == false);
-				REQUIRE(Task(10, true) == 10);
+				REQUIRE(CallTask(Task, 10, true) == 10);
 				REQUIRE(WasLargeLambdaCalled(10));
-				REQUIRE(Task(20, false) == 20);
+				REQUIRE(CallTask(Task, 20, false) == 20);
 				REQUIRE(WasLargeLambdaCalled(-20));
 			}
 			SECTION("Aligned")
@@ -588,9 +601,9 @@ TEST_CASE("TTask")
 				FTask Original(AlignedLambda);
 				FTask Task(MoveTemp(Original));
 				REQUIRE(Original.IsSet() == false);
-				REQUIRE(Task(10, true) == 10);
+				REQUIRE(CallTask(Task, 10, true) == 10);
 				REQUIRE(WasAlignedLambdaCalled(10));
-				REQUIRE(Task(20, false) == 20);
+				REQUIRE(CallTask(Task, 20, false) == 20);
 				REQUIRE(WasAlignedLambdaCalled(-20));
 			}
 		}
diff --git a/Engine/Source/Runtime/AutoRTFM/Private/API.cpp b/Engine/Source/Runtime/AutoRTFM/Private/API.cpp
--- a/Engine/Source/Runtime/AutoRTFM/Private/API.cpp
+++ b/Engine/Source/Runtime/AutoRTFM/Private/API.cpp
@@ -1,6 +1,7 @@
 // Copyright Epic Games, Inc. All Rights Reserved.
 
 #include "AutoRTFM.h"
+#include "AutoRTFMDefines.h"
 #include "Context.h"
 #include "ExternAPI.h"
 #include "Transaction.h"
@@ -429,12 +430,12 @@ namespace
 {
 
 // Internal closed-variant implementations.
-AUTORTFM_CORE_API bool RTFM_autortfm_is_transactional()
+AUTORTFM_INTERNAL bool RTFM_autortfm_is_transactional()
 {
 	return true;
 }
 
-AUTORTFM_CORE_API autortfm_result RTFM_autortfm_transact(void (*UninstrumentedWork)(void*), void (*InstrumentedWork)(void*), void* Arg)
+AUTORTFM_INTERNAL autortfm_result RTFM_autortfm_transact(void (*UninstrumentedWork)(void*), void (*InstrumentedWork)(void*), void* Arg)
 {
 	FContext* Context = FContext::Get();
 	return static_cast<autortfm_result>(Context->Transact(UninstrumentedWork, InstrumentedWork, Arg));
@@ -449,46 +450,46 @@ UE_AUTORTFM_FORCEINLINE autortfm_result TransactThenOpenImpl(void (*Uninstrument
 		}));
 }
 
-AUTORTFM_CORE_API autortfm_result RTFM_autortfm_transact_then_open(void (*UninstrumentedWork)(void*), void (*InstrumentedWork)(void*), void* Arg, const void* ReturnAddress)
+AUTORTFM_INTERNAL autortfm_result RTFM_autortfm_transact_then_open(void (*UninstrumentedWork)(void*), void (*InstrumentedWork)(void*), void* Arg, const void* ReturnAddress)
 {
 	return TransactThenOpenImpl(UninstrumentedWork, Arg, ReturnAddress);
 }
 
-AUTORTFM_CORE_API void RTFM_autortfm_commit(void (*UninstrumentedWork)(void*), void (*InstrumentedWork)(void*), void* Arg)
+AUTORTFM_INTERNAL void RTFM_autortfm_commit(void (*UninstrumentedWork)(void*), void (*InstrumentedWork)(void*), void* Arg)
 {
 	autortfm_result Result = autortfm_transact(UninstrumentedWork, InstrumentedWork, Arg);
 	AUTORTFM_FATAL_IF(Result != autortfm_committed, "Unexpected transaction result: %u", Result);
 }
 
-AUTORTFM_CORE_API void RTFM_autortfm_abort()
+AUTORTFM_INTERNAL void RTFM_autortfm_abort()
 {
 	FContext* Context = FContext::Get();
 	Context->AbortByRequestAndThrow();
 }
 
-AUTORTFM_CORE_API void RTFM_autortfm_start_transaction()
+AUTORTFM_INTERNAL void RTFM_autortfm_start_transaction()
 {
 	AUTORTFM_FATAL("The function `autortfm_start_transaction` was called from closed code");
 }
 
-AUTORTFM_CORE_API void RTFM_autortfm_commit_transaction()
+AUTORTFM_INTERNAL void RTFM_autortfm_commit_transaction()
 {
 	AUTORTFM_FATAL("The function `RTFM_autortfm_commit_transaction` was called from closed code");
 }
 
-AUTORTFM_CORE_API void RTFM_autortfm_abort_transaction()
+AUTORTFM_INTERNAL void RTFM_autortfm_abort_transaction()
 {
 	FContext* const Context = FContext::Get();
 	Context->AbortTransaction(EContextStatus::AbortedByRequest);
 }
 
-AUTORTFM_CORE_API void RTFM_autortfm_clear_transaction_status()
+AUTORTFM_INTERNAL void RTFM_autortfm_clear_transaction_status()
 {
 	AUTORTFM_FATAL("The function `autortfm_clear_transaction_status` was called from closed code");
 	AutoRTFM::InternalUnreachable();
 }
 
-AUTORTFM_CORE_API autortfm_status RTFM_autortfm_close(void (*UninstrumentedWork)(void*), void (*InstrumentedWork)(void*), void* Arg)
+AUTORTFM_INTERNAL autortfm_status RTFM_autortfm_close(void (*UninstrumentedWork)(void*), void (*InstrumentedWork)(void*), void* Arg)
 {
 	FContext* const Context = FContext::Get();
 
@@ -544,7 +545,7 @@ extern "C" UE_AUTORTFM_NOAUTORTFM UE_AUTORTFM_API void autortfm_post_static_loca
 	}
 }
 
-extern "C" AUTORTFM_CORE_API UE_AUTORTFM_API void autortfm_open_from_closed_explicit_validation(autortfm_memory_validation_level MemoryValidationLevel, void (*Work)(void*), void* Arg, const void* ReturnAddress)
+extern "C" AUTORTFM_INTERNAL UE_AUTORTFM_API void autortfm_open_from_closed_explicit_validation(autortfm_memory_validation_level MemoryValidationLevel, void (*Work)(void*), void* Arg, const void* ReturnAddress)
 {
 	FTransaction* Transaction = FContext::Get()->GetCurrentTransaction();
 	if (Transaction)
@@ -570,7 +571,7 @@ extern "C" AUTORTFM_CORE_API UE_AUTORTFM_API void autortfm_open_from_closed_expl
 	}
 }
 
-extern "C" AUTORTFM_CORE_API UE_AUTORTFM_API void autortfm_open_from_closed(void (*Work)(void*), void* Arg, const void* ReturnAddress)
+extern "C" AUTORTFM_INTERNAL UE_AUTORTFM_API void autortfm_open_from_closed(void (*Work)(void*), void* Arg, const void* ReturnAddress)
 {
 	FTransaction* Transaction = FContext::Get()->GetCurrentTransaction();
 	if (Transaction)
@@ -592,116 +593,116 @@ extern "C" AUTORTFM_CORE_API UE_AUTORTFM_API void autortfm_open_from_closed(void
 	}
 }
 
-AUTORTFM_CORE_API void RTFM_autortfm_record_open_write_err(void*, size_t)
+AUTORTFM_INTERNAL void RTFM_autortfm_record_open_write_err(void*, size_t)
 {
 	AUTORTFM_FATAL("The function `autortfm_record_open_write` was called from closed code");
 }
 
-AUTORTFM_CORE_API void RTFM_CascadingAbortTransactionInternal(TTask<void()>&& RunAfterAbort)
+AUTORTFM_INTERNAL void RTFM_CascadingAbortTransactionInternal(TTask<void()>&& RunAfterAbort)
 {
 	FContext* const Context = FContext::Get();
 	AUTORTFM_ASSERT(Context->GetStatus() == EContextStatus::OnTrack);
 	Context->AbortTransactionWithPostAbortCallback(EContextStatus::AbortedByCascadingAbort, std::move(RunAfterAbort));
 }
 
-AUTORTFM_CORE_API void RTFM_CascadingRetryTransactionInternal(TTask<void()> && RunAfterAbortBeforeRetryWork)
+AUTORTFM_INTERNAL void RTFM_CascadingRetryTransactionInternal(TTask<void()> && RunAfterAbortBeforeRetryWork)
 {
 	FContext* Context = FContext::Get();
 	AUTORTFM_ASSERT(Context->GetStatus() == EContextStatus::OnTrack);
 	Context->AbortTransactionWithPostAbortCallback(EContextStatus::AbortedByCascadingRetry, std::move(RunAfterAbortBeforeRetryWork));
 }
 
-AUTORTFM_CORE_API void RTFM_OnCommitInternal(TTask<void()> && Work)
+AUTORTFM_INTERNAL void RTFM_OnCommitInternal(TTask<void()> && Work)
 {
 	FContext* Context = FContext::Get();
 	AUTORTFM_ASSERT(Context->GetStatus() == EContextStatus::OnTrack);
 	Context->GetCurrentTransaction()->DeferUntilCommit(std::move(Work));
 }
 
-AUTORTFM_CORE_API void RTFM_OnAbortInternal(TTask<void()>&& Work)
+AUTORTFM_INTERNAL void RTFM_OnAbortInternal(TTask<void()>&& Work)
 {
 	FContext* Context = FContext::Get();
 	AUTORTFM_ASSERT(Context->GetStatus() == EContextStatus::OnTrack);
 	Context->GetCurrentTransaction()->DeferUntilAbort(std::move(Work));
 }
 
-AUTORTFM_CORE_API void RTFM_PushOnCommitHandlerInternal(const void* Key, TTask<void()>&& Work)
+AUTORTFM_INTERNAL void RTFM_PushOnCommitHandlerInternal(const void* Key, TTask<void()>&& Work)
 {
 	FContext* Context = FContext::Get();
 	AUTORTFM_ASSERT(Context->GetStatus() == EContextStatus::OnTrack);
 	Context->GetCurrentTransaction()->PushDeferUntilCommitHandler(Key, std::move(Work));
 }
 
-AUTORTFM_CORE_API void RTFM_PopOnCommitHandlerInternal(const void* Key)
+AUTORTFM_INTERNAL void RTFM_PopOnCommitHandlerInternal(const void* Key)
 {
 	FContext* Context = FContext::Get();
 	AUTORTFM_ASSERT(Context->GetStatus() == EContextStatus::OnTrack);
 	Context->GetCurrentTransaction()->PopDeferUntilCommitHandler(Key);
 }
 
-AUTORTFM_CORE_API void RTFM_PopAllOnCommitHandlersInternal(const void* Key)
+AUTORTFM_INTERNAL void RTFM_PopAllOnCommitHandlersInternal(const void* Key)
 {
 	FContext* Context = FContext::Get();
 	AUTORTFM_ASSERT(Context->GetStatus() == EContextStatus::OnTrack);
 	Context->GetCurrentTransaction()->PopAllDeferUntilCommitHandlers(Key);
 }
 
-AUTORTFM_CORE_API void RTFM_PushOnAbortHandlerInternal(const void* Key, TTask<void()>&& Work)
+AUTORTFM_INTERNAL void RTFM_PushOnAbortHandlerInternal(const void* Key, TTask<void()>&& Work)
 {
 	FContext* Context = FContext::Get();
 	AUTORTFM_ASSERT(Context->GetStatus() == EContextStatus::OnTrack);
 	Context->GetCurrentTransaction()->PushDeferUntilAbortHandler(Key, std::move(Work));
 }
 
-AUTORTFM_CORE_API void RTFM_PopOnAbortHandlerInternal(const void* Key)
+AUTORTFM_INTERNAL void RTFM_PopOnAbortHandlerInternal(const void* Key)
 {
 	FContext* Context = FContext::Get();
 	AUTORTFM_ASSERT(Context->GetStatus() == EContextStatus::OnTrack);
 	Context->GetCurrentTransaction()->PopDeferUntilAbortHandler(Key);
 }
 
-AUTORTFM_CORE_API void RTFM_PopAllOnAbortHandlersInternal(const void* Key)
+AUTORTFM_INTERNAL void RTFM_PopAllOnAbortHandlersInternal(const void* Key)
 {
 	FContext* Context = FContext::Get();
 	AUTORTFM_ASSERT(Context->GetStatus() == EContextStatus::OnTrack);
 	Context->GetCurrentTransaction()->PopAllDeferUntilAbortHandlers(Key);
 }
 
-AUTORTFM_CORE_API void RTFM_autortfm_on_commit(void (*Work)(void*), void* Arg)
+AUTORTFM_INTERNAL void RTFM_autortfm_on_commit(void (*Work)(void*), void* Arg)
 {
 	RTFM_OnCommitInternal([Work, Arg] { Work(Arg); });
 }
 
-AUTORTFM_CORE_API void RTFM_autortfm_on_abort(void (*Work)(void*), void* Arg)
+AUTORTFM_INTERNAL void RTFM_autortfm_on_abort(void (*Work)(void*), void* Arg)
 {
 	RTFM_OnAbortInternal([Work, Arg] { Work(Arg); });
 }
 
-AUTORTFM_CORE_API void RTFM_autortfm_push_on_abort_handler(const void* Key, void (*Work)(void*), void* Arg)
+AUTORTFM_INTERNAL void RTFM_autortfm_push_on_abort_handler(const void* Key, void (*Work)(void*), void* Arg)
 {
 	RTFM_PushOnAbortHandlerInternal(Key, [Work, Arg] { Work(Arg); });
 }
 
-AUTORTFM_CORE_API void RTFM_autortfm_pop_on_abort_handler(const void* Key)
+AUTORTFM_INTERNAL void RTFM_autortfm_pop_on_abort_handler(const void* Key)
 {
 	RTFM_PopOnAbortHandlerInternal(Key);
 }
 
-AUTORTFM_CORE_API void* RTFM_autortfm_did_allocate(void* Ptr, size_t Size)
+AUTORTFM_INTERNAL void* RTFM_autortfm_did_allocate(void* Ptr, size_t Size)
 {
 	FContext* Context = FContext::Get();
 	Context->DidAllocate(Ptr, Size);
 	return Ptr;
 }
 
-AUTORTFM_CORE_API void RTFM_autortfm_did_free(void* Ptr)
+AUTORTFM_INTERNAL void RTFM_autortfm_did_free(void* Ptr)
 {
 	// We should never-ever-ever actually free memory from within closed code of
 	// a transaction.
 	AutoRTFM::InternalUnreachable();
 }
 
-AUTORTFM_CORE_API bool IsAutoRTFMInitialized()
+AUTORTFM_INTERNAL bool IsAutoRTFMInitialized()
 {
 	return FContext::Get() != nullptr;
 }
@@ -709,7 +710,7 @@ AUTORTFM_CORE_API bool IsAutoRTFMInitialized()
 // Consume the GPendingOpenToClosedRegistrations linked list to register the
 // open -> closed functions. This is done via a linked-list to avoid heap
 // allocations before AutoRTFM is initialized.
-AUTORTFM_CORE_API void ProcessAllPendingOpenToClosedRegistrations()
+AUTORTFM_INTERNAL void ProcessAllPendingOpenToClosedRegistrations()
 {
 	AUTORTFM_ASSERT(IsAutoRTFMInitialized());
 	FunctionMapAdd(GPendingOpenToClosedRegistrations);
@@ -722,6 +723,7 @@ AUTORTFM_CORE_API void ProcessAllPendingOpenToClosedRegistrations()
 
 extern "C" void autortfm_initialize(const autortfm_extern_api* ExternAPI) noexcept
 {
+	AutoRTFM::UnreachableIfClosed("TODO: Mark this as AUTORTFM_DISABLED and fixup callers");
 	AUTORTFM_ENSURE_MSG(!FContext::Get(), "AutoRTFM initialized twice");
 	AUTORTFM_ASSERT(ExternAPI);
 	AUTORTFM_ASSERT(ExternAPI->Allocate);
@@ -751,7 +753,7 @@ extern "C" bool autortfm_is_transactional() noexcept
 	return false;
 }
 
-extern "C" bool autortfm_is_committing_or_aborting() noexcept
+extern "C" UE_AUTORTFM_ALWAYS_OPEN bool autortfm_is_committing_or_aborting() noexcept
 {
 	if (ForTheRuntime::IsAutoRTFMRuntimeEnabled())
 	{
@@ -870,7 +872,7 @@ extern "C" void autortfm_clear_transaction_status() noexcept
 	FContext::Get()->ClearTransactionStatus();
 }
 
-extern "C" autortfm_status autortfm_get_context_status() noexcept
+extern "C" UE_AUTORTFM_ALWAYS_OPEN autortfm_status autortfm_get_context_status() noexcept
 {
 	return static_cast<autortfm_status>(FContext::Get()->GetStatus());
 }
@@ -1022,7 +1024,7 @@ extern "C" UE_AUTORTFM_NOAUTORTFM void autortfm_unregister_open_to_closed_functi
 	// then something will need to be done here.
 }
 
-extern "C" bool autortfm_is_on_current_transaction_stack(void* Ptr) noexcept
+extern "C" UE_AUTORTFM_ALWAYS_OPEN bool autortfm_is_on_current_transaction_stack(void* Ptr) noexcept
 {
 	if (FContext* Context = FContext::Get())
 	{
diff --git a/Engine/Source/Runtime/AutoRTFM/Private/CompilerInterface.cpp b/Engine/Source/Runtime/AutoRTFM/Private/CompilerInterface.cpp
--- a/Engine/Source/Runtime/AutoRTFM/Private/CompilerInterface.cpp
+++ b/Engine/Source/Runtime/AutoRTFM/Private/CompilerInterface.cpp
@@ -1,5 +1,6 @@
 // Copyright Epic Games, Inc. All Rights Reserved.
 
+#include "AutoRTFMDefines.h"
 #if (defined(__AUTORTFM) && __AUTORTFM)
 
 #include "AutoRTFM.h"
@@ -12,25 +13,6 @@
 namespace AutoRTFM
 {
 
-void AbortDueToBadAlignment(FContext* Context, void* Ptr, size_t Alignment, const char* Message = nullptr)
-{
-    Context->DumpState();
-    fprintf(stderr, "Aborting because alignment error: expected alignment %zu, got pointer %p.\n", Alignment, Ptr);
-    if (Message)
-    {
-        fprintf(stderr, "%s\n", Message);
-    }
-    abort();
-}
-
-void CheckAlignment(FContext* Context, void* Ptr, size_t AlignmentMask)
-{
-    if (reinterpret_cast<uintptr_t>(Ptr) & AlignmentMask)
-    {
-        AbortDueToBadAlignment(Context, Ptr, AlignmentMask + 1);
-    }
-}
-
 #if !AUTORTFM_BUILD_SHIPPING
 
 // Check for writes to null in development code, so that the inevitable crash will occur
@@ -54,42 +36,42 @@ void CheckAlignment(FContext* Context, void* Ptr, size_t AlignmentMask)
 
 #endif
 
-extern "C" UE_AUTORTFM_API AUTORTFM_NO_ASAN void autortfm_record_write(void* Ptr, size_t Size) noexcept
+extern "C" AUTORTFM_DISABLE UE_AUTORTFM_API AUTORTFM_NO_ASAN void autortfm_record_write(void* Ptr, size_t Size) noexcept
 {
 	UE_AUTORTFM_HANDLE_NULL_WRITE(Ptr);
 	FContext* Context = FContext::Get();
 	Context->RecordWrite(Ptr, Size);
 }
 
-extern "C" UE_AUTORTFM_API AUTORTFM_NO_ASAN void autortfm_record_write_1(void* Ptr) noexcept
+extern "C" AUTORTFM_DISABLE UE_AUTORTFM_API AUTORTFM_NO_ASAN void autortfm_record_write_1(void* Ptr) noexcept
 {
 	UE_AUTORTFM_HANDLE_NULL_WRITE(Ptr);
 	FContext* Context = FContext::Get();
 	Context->RecordWrite<1>(Ptr);
 }
 
-extern "C" UE_AUTORTFM_API AUTORTFM_NO_ASAN void autortfm_record_write_2(void* Ptr) noexcept
+extern "C" AUTORTFM_DISABLE UE_AUTORTFM_API AUTORTFM_NO_ASAN void autortfm_record_write_2(void* Ptr) noexcept
 {
 	UE_AUTORTFM_HANDLE_NULL_WRITE(Ptr);
 	FContext* Context = FContext::Get();
 	Context->RecordWrite<2>(Ptr);
 }
 
-extern "C" UE_AUTORTFM_API AUTORTFM_NO_ASAN void autortfm_record_write_4(void* Ptr) noexcept
+extern "C" AUTORTFM_DISABLE UE_AUTORTFM_API AUTORTFM_NO_ASAN void autortfm_record_write_4(void* Ptr) noexcept
 {
 	UE_AUTORTFM_HANDLE_NULL_WRITE(Ptr);
 	FContext* Context = FContext::Get();
 	Context->RecordWrite<4>(Ptr);
 }
 
-extern "C" UE_AUTORTFM_API AUTORTFM_NO_ASAN void autortfm_record_write_8(void* Ptr) noexcept
+extern "C" AUTORTFM_DISABLE UE_AUTORTFM_API AUTORTFM_NO_ASAN void autortfm_record_write_8(void* Ptr) noexcept
 {
 	UE_AUTORTFM_HANDLE_NULL_WRITE(Ptr);
 	FContext* Context = FContext::Get();
 	Context->RecordWrite<8>(Ptr);
 }
 
-extern "C" UE_AUTORTFM_API AUTORTFM_NO_ASAN void autortfm_record_masked_write(void* Ptr, uintptr_t Mask, int MaskWidthBits, int ValueSizeBytes) noexcept
+extern "C" AUTORTFM_DISABLE UE_AUTORTFM_API AUTORTFM_NO_ASAN void autortfm_record_masked_write(void* Ptr, uintptr_t Mask, int MaskWidthBits, int ValueSizeBytes) noexcept
 {
 	UE_AUTORTFM_HANDLE_NULL_WRITE(Ptr);
 
@@ -118,31 +100,31 @@ extern "C" UE_AUTORTFM_NOAUTORTFM UE_AUTORTFM_API void* autortfm_lookup_function
 	return FunctionMapLookup(OriginalFunction, Where);
 }
 
-extern "C" UE_AUTORTFM_API void autortfm_memcpy(void* Dst, const void* Src, size_t Size) noexcept
+extern "C" AUTORTFM_DISABLE UE_AUTORTFM_API void autortfm_memcpy(void* Dst, const void* Src, size_t Size) noexcept
 {
 	FContext* Context = FContext::Get();
     Memcpy(Dst, Src, Size, Context);
 }
 
-extern "C" UE_AUTORTFM_API void autortfm_memmove(void* Dst, const void* Src, size_t Size) noexcept
+extern "C" AUTORTFM_DISABLE UE_AUTORTFM_API void autortfm_memmove(void* Dst, const void* Src, size_t Size) noexcept
 {
 	FContext* Context = FContext::Get();
     Memmove(Dst, Src, Size, Context);
 }
 
-extern "C" UE_AUTORTFM_API void autortfm_memset(void* Dst, int Value, size_t Size) noexcept
+extern "C" AUTORTFM_DISABLE UE_AUTORTFM_API void autortfm_memset(void* Dst, int Value, size_t Size) noexcept
 {
 	FContext* Context = FContext::Get();
     Memset(Dst, Value, Size, Context);
 }
 
-extern "C" UE_AUTORTFM_API void autortfm_unreachable(const char* Message) noexcept
+extern "C" UE_AUTORTFM_ALWAYS_OPEN UE_AUTORTFM_API void autortfm_unreachable(const char* Message) noexcept
 {
 	AUTORTFM_REPORT_ERROR("AutoRTFM Unreachable: %s", Message);
 	__builtin_unreachable();
 }
 
-extern "C" UE_AUTORTFM_API void autortfm_llvm_fail(const char* Message) noexcept
+extern "C" AUTORTFM_DISABLE UE_AUTORTFM_API void autortfm_llvm_fail(const char* Message) noexcept
 {
 	if (Message)
 	{
@@ -154,7 +136,7 @@ extern "C" UE_AUTORTFM_API void autortfm_llvm_fail(const char* Message) noexcept
 	}
 }
 
-extern "C" UE_AUTORTFM_API void autortfm_llvm_missing_function() noexcept
+extern "C" AUTORTFM_DISABLE UE_AUTORTFM_API void autortfm_llvm_missing_function() noexcept
 {
 	if (ForTheRuntime::GetInternalAbortAction() == ForTheRuntime::EAutoRTFMInternalAbortActionState::Crash)
 	{
@@ -169,12 +151,12 @@ extern "C" UE_AUTORTFM_API void autortfm_llvm_missing_function() noexcept
     Context->AbortByLanguageAndThrow();
 }
 
-extern "C" UE_AUTORTFM_API void autortfm_called_no_autortfm() noexcept
+extern "C" AUTORTFM_DISABLE UE_AUTORTFM_API void autortfm_called_no_autortfm() noexcept
 {
 	AUTORTFM_FATAL("inlined UE_AUTORTFM_NOAUTORTFM function called from the closed");
 }
 
-extern "C" UE_AUTORTFM_API void autortfm_disable_called() noexcept
+extern "C" AUTORTFM_DISABLE UE_AUTORTFM_API void autortfm_disable_called() noexcept
 {
 	AUTORTFM_FATAL("inlined AUTORTFM_DISABLE function called from the closed");
 }
diff --git a/Engine/Source/Runtime/AutoRTFM/Private/Context.h b/Engine/Source/Runtime/AutoRTFM/Private/Context.h
--- a/Engine/Source/Runtime/AutoRTFM/Private/Context.h
+++ b/Engine/Source/Runtime/AutoRTFM/Private/Context.h
@@ -18,7 +18,7 @@ namespace AutoRTFM
 class FLineLock;
 class FCallNest;
 
-class AUTORTFM_CORE_API FContext final
+class AUTORTFM_DISABLE FContext final
 {
 	using FTaskPool = typename TTaskArray<TTask<void()>>::FEntryPool;
 	using FTransactionPool = TIntrusivePool<FTransaction, 16>;
@@ -98,7 +98,7 @@ public:
 
 	inline FOpenHashThrottler& GetOpenHashThrottler() { return OpenHashThrottler; }
 
-    void DumpState() const;
+	void DumpState() const;
 
 	inline void EnteringStaticLocalInitializer()
 	{
diff --git a/Engine/Source/Runtime/AutoRTFM/Private/FunctionMap.cpp b/Engine/Source/Runtime/AutoRTFM/Private/FunctionMap.cpp
--- a/Engine/Source/Runtime/AutoRTFM/Private/FunctionMap.cpp
+++ b/Engine/Source/Runtime/AutoRTFM/Private/FunctionMap.cpp
@@ -251,7 +251,7 @@ inline void* FollowRelocation(void* Function)
 	return nullptr;
 }
 
-void* FunctionMapReportError(void* Open, const char* Where)
+AUTORTFM_DISABLE static void* FunctionMapReportError(void* Open, const char* Where)
 {
 	std::string FunctionDescription = GetFunctionDescription(Open);
 	if (Where)
diff --git a/Engine/Source/Runtime/AutoRTFM/Private/FunctionMap.h b/Engine/Source/Runtime/AutoRTFM/Private/FunctionMap.h
--- a/Engine/Source/Runtime/AutoRTFM/Private/FunctionMap.h
+++ b/Engine/Source/Runtime/AutoRTFM/Private/FunctionMap.h
@@ -10,17 +10,17 @@ namespace AutoRTFM
 {
 
 // Dumps statistics about the function map to LOG_INFO.
-void FunctionMapDumpStats();
+AUTORTFM_DISABLE void FunctionMapDumpStats();
 
 // Adds all the open to closed function mappings from the linked list of tables to the function map.
-void FunctionMapAdd(autortfm_open_to_closed_table* Tables);
+AUTORTFM_DISABLE void FunctionMapAdd(autortfm_open_to_closed_table* Tables);
 
 // Looks up the closed function from the open function using the function map, following dynamic
 // library thunk redirections.
-void* FunctionMapLookupExhaustive(void* OpenFn, const char* Where);
+AUTORTFM_DISABLE void* FunctionMapLookupExhaustive(void* OpenFn, const char* Where);
 
 template<typename ReturnType, typename... ParameterTypes>
-auto FunctionMapLookupExhaustive(ReturnType (*Function)(ParameterTypes...), const char* Where) -> ReturnType (*)(ParameterTypes...)
+AUTORTFM_DISABLE auto FunctionMapLookupExhaustive(ReturnType (*Function)(ParameterTypes...), const char* Where) -> ReturnType (*)(ParameterTypes...)
 {
 	return reinterpret_cast<ReturnType (*)(ParameterTypes...)>(FunctionMapLookupExhaustive(reinterpret_cast<void*>(Function), Where));
 }
diff --git a/Engine/Source/Runtime/AutoRTFM/Private/FunctionMapInlines.h b/Engine/Source/Runtime/AutoRTFM/Private/FunctionMapInlines.h
--- a/Engine/Source/Runtime/AutoRTFM/Private/FunctionMapInlines.h
+++ b/Engine/Source/Runtime/AutoRTFM/Private/FunctionMapInlines.h
@@ -4,6 +4,7 @@
 
 #if (defined(__AUTORTFM) && __AUTORTFM)
 
+#include "AutoRTFMDefines.h"
 #include "AutoRTFMConstants.h"
 #include "FunctionMap.h"
 #include "Utils.h"
@@ -11,7 +12,7 @@
 namespace AutoRTFM
 {
 
-UE_AUTORTFM_FORCEINLINE
+AUTORTFM_DISABLE UE_AUTORTFM_FORCEINLINE
 void* FunctionMapLookupUsingMagicPrefix(void* OpenFn)
 {
 	// We use prefix data in our custom LLVM pass to stuff some data just
@@ -42,7 +43,7 @@ void* FunctionMapLookupUsingMagicPrefix(void* OpenFn)
 	return nullptr;
 }
 
-inline void* FunctionMapLookup(void* OpenFn, const char* Where)
+AUTORTFM_DISABLE inline void* FunctionMapLookup(void* OpenFn, const char* Where)
 {
 	if (void* ClosedFn = FunctionMapLookupUsingMagicPrefix(OpenFn); AUTORTFM_LIKELY(ClosedFn))
 	{
@@ -53,7 +54,7 @@ inline void* FunctionMapLookup(void* OpenFn, const char* Where)
 }
 
 template<typename TReturnType, typename... TParameterTypes>
-auto FunctionMapLookup(TReturnType (*OpenFn)(TParameterTypes...), const char* Where) -> TReturnType (*)(TParameterTypes...)
+AUTORTFM_DISABLE auto FunctionMapLookup(TReturnType (*OpenFn)(TParameterTypes...), const char* Where) -> TReturnType (*)(TParameterTypes...)
 {
     return reinterpret_cast<TReturnType (*)(TParameterTypes...)>(FunctionMapLookup(reinterpret_cast<void*>(OpenFn), Where));
 }
diff --git a/Engine/Source/Runtime/AutoRTFM/Private/IntrusivePool.h b/Engine/Source/Runtime/AutoRTFM/Private/IntrusivePool.h
--- a/Engine/Source/Runtime/AutoRTFM/Private/IntrusivePool.h
+++ b/Engine/Source/Runtime/AutoRTFM/Private/IntrusivePool.h
@@ -4,6 +4,7 @@
 
 #if (defined(__AUTORTFM) && __AUTORTFM)
 
+#include "AutoRTFMDefines.h"
 #include "BlockAllocator.h"
 
 #include <cstddef>
@@ -13,7 +14,7 @@ namespace AutoRTFM
 {
 
 template<typename ItemType, size_t InlineCapacity>
-class TIntrusivePool final
+class AUTORTFM_DISABLE TIntrusivePool final
 {
 public:
 	using FItem = ItemType;
diff --git a/Engine/Source/Runtime/AutoRTFM/Private/Memcpy.h b/Engine/Source/Runtime/AutoRTFM/Private/Memcpy.h
--- a/Engine/Source/Runtime/AutoRTFM/Private/Memcpy.h
+++ b/Engine/Source/Runtime/AutoRTFM/Private/Memcpy.h
@@ -4,6 +4,8 @@
 
 #if (defined(__AUTORTFM) && __AUTORTFM)
 
+#include "AutoRTFMDefines.h"
+
 #include <stdlib.h>
 
 namespace AutoRTFM
@@ -11,10 +13,10 @@ namespace AutoRTFM
 
 class FContext;
 
-void* MemcpyToNew(void* Dst, const void* Src, size_t Size, FContext* Context);
-void* Memcpy(void* Dst, const void* Src, size_t Size, FContext* Context);
-void* Memmove(void* Dst, const void* Src, size_t Size, FContext* Context);
-void* Memset(void* Dst, int Value, size_t Size, FContext* Context);
+AUTORTFM_DISABLE void* MemcpyToNew(void* Dst, const void* Src, size_t Size, FContext* Context);
+AUTORTFM_DISABLE void* Memcpy(void* Dst, const void* Src, size_t Size, FContext* Context);
+AUTORTFM_DISABLE void* Memmove(void* Dst, const void* Src, size_t Size, FContext* Context);
+AUTORTFM_DISABLE void* Memset(void* Dst, int Value, size_t Size, FContext* Context);
 
 } // namespace AutoRTFM
 
diff --git a/Engine/Source/Runtime/AutoRTFM/Private/StdLibSupport.cpp b/Engine/Source/Runtime/AutoRTFM/Private/StdLibSupport.cpp
--- a/Engine/Source/Runtime/AutoRTFM/Private/StdLibSupport.cpp
+++ b/Engine/Source/Runtime/AutoRTFM/Private/StdLibSupport.cpp
@@ -64,7 +64,7 @@ namespace
 
 // A helper that opens a FILE to "/dev/null" on first call to Get()
 // and automatically closes the file on static destruction.
-class FNullFile
+class AUTORTFM_INTERNAL FNullFile
 {
 public:
     static FILE* Get()
@@ -79,7 +79,7 @@ private:
     FILE* const File;
 };
 
-void ThrowErrorFormatContainsPercentN()
+AUTORTFM_INTERNAL void ThrowErrorFormatContainsPercentN()
 {
     AUTORTFM_WARN("AutoRTFM does not support format strings containing '%%n'");
     FContext* Context = FContext::Get();
@@ -87,7 +87,7 @@ void ThrowErrorFormatContainsPercentN()
 }
 
 // Throws an error if the format string contains a '%n'.
-static void ThrowIfFormatContainsPercentN(const char* Format)
+AUTORTFM_INTERNAL static void ThrowIfFormatContainsPercentN(const char* Format)
 {
     for (const char* P = Format; *P != '\0'; ++P)
     {
@@ -106,7 +106,7 @@ static void ThrowIfFormatContainsPercentN(const char* Format)
 }
 
 // Throws an error if the format string contains a '%n'.
-static void ThrowIfFormatContainsPercentN(const wchar_t* Format)
+AUTORTFM_INTERNAL static void ThrowIfFormatContainsPercentN(const wchar_t* Format)
 {
     for (const wchar_t* P = Format; *P != L'\0'; ++P)
     {
@@ -127,7 +127,7 @@ static void ThrowIfFormatContainsPercentN(const wchar_t* Format)
 #if AUTORTFM_PLATFORM_WINDOWS
 // vvvvvvvvvvvvvvvvvvvvvvvvvvvvvv AUTORTFM_PLATFORM_WINDOWS vvvvvvvvvvvvvvvvvvvvvvvvvvvvv
 template<typename RETURN, typename CHAR, RETURN FN(const CHAR*, CHAR**, _locale_t)>
-RETURN RTFM_StringToFloat(const CHAR* String, CHAR** EndPtr, _locale_t Locale)
+AUTORTFM_INTERNAL RETURN RTFM_StringToFloat(const CHAR* String, CHAR** EndPtr, _locale_t Locale)
 {
     if (nullptr != EndPtr)
     {
@@ -137,7 +137,7 @@ RETURN RTFM_StringToFloat(const CHAR* String, CHAR** EndPtr, _locale_t Locale)
 }
 
 template<typename RETURN, typename CHAR, RETURN FN(const CHAR*, CHAR**, int)>
-RETURN RTFM_StringToIntNoLocale(const CHAR* String, CHAR** EndPtr, int Radix)
+AUTORTFM_INTERNAL RETURN RTFM_StringToIntNoLocale(const CHAR* String, CHAR** EndPtr, int Radix)
 {
     if (nullptr != EndPtr)
     {
@@ -147,7 +147,7 @@ RETURN RTFM_StringToIntNoLocale(const CHAR* String, CHAR** EndPtr, int Radix)
 }
 
 template<typename RETURN, typename CHAR, RETURN FN(const CHAR*, CHAR**, int, _locale_t)>
-RETURN RTFM_StringToInt(const CHAR* String, CHAR** EndPtr, int Radix, _locale_t Locale)
+AUTORTFM_INTERNAL RETURN RTFM_StringToInt(const CHAR* String, CHAR** EndPtr, int Radix, _locale_t Locale)
 {
     if (nullptr != EndPtr)
     {
@@ -156,7 +156,7 @@ RETURN RTFM_StringToInt(const CHAR* String, CHAR** EndPtr, int Radix, _locale_t
     return FN(String, EndPtr, Radix, Locale);
 }
 
-FILE* RTFM___acrt_iob_func(unsigned Index)
+AUTORTFM_INTERNAL FILE* RTFM___acrt_iob_func(unsigned Index)
 {
     switch (Index)
     {
@@ -174,7 +174,7 @@ FILE* RTFM___acrt_iob_func(unsigned Index)
 }
 
 // FIXME: Does not currently support %n format specifiers.
-int RTFM___stdio_common_vfprintf(
+AUTORTFM_INTERNAL int RTFM___stdio_common_vfprintf(
         unsigned __int64 Options,
         FILE*            Stream,
         char const*      Format,
@@ -187,7 +187,7 @@ int RTFM___stdio_common_vfprintf(
 }
 
 // FIXME: Does not currently support %n format specifiers.
-int RTFM___stdio_common_vsprintf(
+AUTORTFM_INTERNAL int RTFM___stdio_common_vsprintf(
         unsigned __int64 Options,
         char*            Buffer,
         size_t           BufferCount,
@@ -216,7 +216,7 @@ int RTFM___stdio_common_vsprintf(
 }
 
 // FIXME: Does not currently support %n format specifiers.
-int RTFM___stdio_common_vswprintf(
+AUTORTFM_INTERNAL int RTFM___stdio_common_vswprintf(
         unsigned __int64 Options,
         wchar_t*         Buffer,
         size_t           BufferCount,
@@ -245,7 +245,7 @@ int RTFM___stdio_common_vswprintf(
 }
 
 // FIXME: Does not currently support %n format specifiers.
-int RTFM___stdio_common_vfwprintf(
+AUTORTFM_INTERNAL int RTFM___stdio_common_vfwprintf(
         unsigned __int64 Options,
         FILE*            Stream,
         wchar_t const*   Format,
@@ -257,7 +257,7 @@ int RTFM___stdio_common_vfwprintf(
     return __stdio_common_vfwprintf(Options, Stream, Format, Locale, ArgList);
 }
 
-BOOL RTFM_TlsSetValue(DWORD dwTlsIndex, LPVOID lpTlsValue)
+AUTORTFM_INTERNAL BOOL RTFM_TlsSetValue(DWORD dwTlsIndex, LPVOID lpTlsValue)
 {
 	LPVOID OriginalValue = TlsGetValue(dwTlsIndex);
 
@@ -290,7 +290,7 @@ extern "C" size_t _ZNSt3__112__next_primeEm(size_t N) __attribute__((weak));
 
 #if AUTORTFM_PLATFORM_LINUX
 // vvvvvvvvvvvvvvvvvvvvvvvvvvvvvv AUTORTFM_PLATFORM_LINUX vvvvvvvvvvvvvvvvvvvvvvvvvvvvvv
-int RTFM_stat(const char* Path, struct stat* StatBuf) throw()
+AUTORTFM_INTERNAL int RTFM_stat(const char* Path, struct stat* StatBuf) throw()
 {
 	FContext* Context = FContext::Get();
 	Context->RecordWrite(StatBuf, sizeof(*StatBuf));
@@ -298,7 +298,7 @@ int RTFM_stat(const char* Path, struct stat* StatBuf) throw()
 	return stat(Path, StatBuf);
 }
 
-int RTFM_fstat(int Fd, struct stat* StatBuf) throw()
+AUTORTFM_INTERNAL int RTFM_fstat(int Fd, struct stat* StatBuf) throw()
 {
 	FContext* Context = FContext::Get();
 	Context->RecordWrite(StatBuf, sizeof(*StatBuf));
@@ -306,7 +306,7 @@ int RTFM_fstat(int Fd, struct stat* StatBuf) throw()
 	return fstat(Fd, StatBuf);
 }
 
-int RTFM___xstat(int Ver, const char* Path, struct stat* StatBuf) throw()
+AUTORTFM_INTERNAL int RTFM___xstat(int Ver, const char* Path, struct stat* StatBuf) throw()
 {
 	FContext* Context = FContext::Get();
 	Context->RecordWrite(StatBuf, sizeof(*StatBuf));
@@ -314,7 +314,7 @@ int RTFM___xstat(int Ver, const char* Path, struct stat* StatBuf) throw()
 	return __xstat(Ver, Path, StatBuf);
 }
 
-int RTFM___fxstat(int Ver, int Fd, struct stat* StatBuf) throw()
+AUTORTFM_INTERNAL int RTFM___fxstat(int Ver, int Fd, struct stat* StatBuf) throw()
 {
 	FContext* Context = FContext::Get();
 	Context->RecordWrite(StatBuf, sizeof(*StatBuf));
@@ -324,22 +324,22 @@ int RTFM___fxstat(int Ver, int Fd, struct stat* StatBuf) throw()
 // ^^^^^^^^^^^^^^^^^^^^^^^^^^^^^^ AUTORTFM_PLATFORM_LINUX ^^^^^^^^^^^^^^^^^^^^^^^^^^^^^^
 #endif
 
-void* RTFM_memcpy(void* Dst, const void* Src, size_t Size) throw()
+AUTORTFM_INTERNAL void* RTFM_memcpy(void* Dst, const void* Src, size_t Size) throw()
 {
     return Memcpy(Dst, Src, Size, FContext::Get());
 }
 
-void* RTFM_memmove(void* Dst, const void* Src, size_t Size) throw()
+AUTORTFM_INTERNAL void* RTFM_memmove(void* Dst, const void* Src, size_t Size) throw()
 {
     return Memmove(Dst, Src, Size, FContext::Get());
 }
 
-void* RTFM_memset(void* Dst, int Value, size_t Size) throw()
+AUTORTFM_INTERNAL void* RTFM_memset(void* Dst, int Value, size_t Size) throw()
 {
     return Memset(Dst, Value, Size, FContext::Get());
 }
 
-void* RTFM_malloc(size_t Size) throw()
+AUTORTFM_INTERNAL void* RTFM_malloc(size_t Size) throw()
 {
     void* Result = malloc(Size);
 	FContext* Context = FContext::Get();
@@ -351,7 +351,7 @@ void* RTFM_malloc(size_t Size) throw()
     return Result;
 }
 
-void* RTFM_calloc(size_t Count, size_t Size) throw()
+AUTORTFM_INTERNAL void* RTFM_calloc(size_t Count, size_t Size) throw()
 {
     void* Result = calloc(Count, Size);
 	FContext* Context = FContext::Get();
@@ -363,7 +363,7 @@ void* RTFM_calloc(size_t Count, size_t Size) throw()
     return Result;
 }
 
-void RTFM_free(void* Ptr) throw()
+AUTORTFM_INTERNAL void RTFM_free(void* Ptr) throw()
 {
     if (Ptr)
     {
@@ -375,7 +375,7 @@ void RTFM_free(void* Ptr) throw()
     }
 }
 
-void* RTFM_realloc(void* Ptr, size_t Size) throw()
+AUTORTFM_INTERNAL void* RTFM_realloc(void* Ptr, size_t Size) throw()
 {
     void* NewObject = RTFM_malloc(Size);
     if (Ptr)
@@ -394,7 +394,7 @@ void* RTFM_realloc(void* Ptr, size_t Size) throw()
     return NewObject;
 }
 
-char* RTFM_strcpy(char* const Dst, const char* const Src) throw()
+AUTORTFM_INTERNAL char* RTFM_strcpy(char* const Dst, const char* const Src) throw()
 {
     const size_t SrcLen = strlen(Src);
 
@@ -403,14 +403,14 @@ char* RTFM_strcpy(char* const Dst, const char* const Src) throw()
     return strcpy(Dst, Src);
 }
 
-char* RTFM_strncpy(char* const Dst, const char* const Src, const size_t Num) throw()
+AUTORTFM_INTERNAL char* RTFM_strncpy(char* const Dst, const char* const Src, const size_t Num) throw()
 {
 	FContext* Context = FContext::Get();
     Context->RecordWrite(Dst, Num);
     return strncpy(Dst, Src, Num);
 }
 
-char* RTFM_strcat(char* const Dst, const char* const Src) throw()
+AUTORTFM_INTERNAL char* RTFM_strcat(char* const Dst, const char* const Src) throw()
 {
     const size_t DstLen = strlen(Dst);
     const size_t SrcLen = strlen(Src);
@@ -420,7 +420,7 @@ char* RTFM_strcat(char* const Dst, const char* const Src) throw()
     return strcat(Dst, Src);
 }
 
-char* RTFM_strncat(char* const Dst, const char* const Src, const size_t Num) throw()
+AUTORTFM_INTERNAL char* RTFM_strncat(char* const Dst, const char* const Src, const size_t Num) throw()
 {
     const size_t DstLen = strlen(Dst);
 
@@ -430,7 +430,7 @@ char* RTFM_strncat(char* const Dst, const char* const Src, const size_t Num) thr
 }
 
 template<typename RETURN, typename CHAR, RETURN FN(const CHAR*, CHAR**)>
-RETURN RTFM_StringToFloat(const CHAR* String, CHAR** EndPtr) throw()
+AUTORTFM_INTERNAL RETURN RTFM_StringToFloat(const CHAR* String, CHAR** EndPtr) throw()
 {
     if (nullptr != EndPtr)
     {
@@ -440,7 +440,7 @@ RETURN RTFM_StringToFloat(const CHAR* String, CHAR** EndPtr) throw()
 }
 
 template<typename RETURN, typename CHAR, RETURN FN(const CHAR*, CHAR**, int)>
-RETURN RTFM_StringToInt(const CHAR* String, CHAR** EndPtr, int Radix) throw()
+AUTORTFM_INTERNAL RETURN RTFM_StringToInt(const CHAR* String, CHAR** EndPtr, int Radix) throw()
 {
     if (nullptr != EndPtr)
     {
@@ -450,35 +450,35 @@ RETURN RTFM_StringToInt(const CHAR* String, CHAR** EndPtr, int Radix) throw()
 }
 
 template<typename T, std::to_chars_result FN(char*, char*, T)>
-std::to_chars_result RTFM_ToChars(char* First, char* Last, T Value)
+AUTORTFM_INTERNAL std::to_chars_result RTFM_ToChars(char* First, char* Last, T Value)
 {
     AutoRTFM::RecordOpenWrite(First, Last - First);
     return FN(First, Last, Value);
 }
 
 template<typename T, std::to_chars_result FN(char*, char*, T, int)>
-std::to_chars_result RTFM_ToChars(char* First, char* Last, T Value, int Base)
+AUTORTFM_INTERNAL std::to_chars_result RTFM_ToChars(char* First, char* Last, T Value, int Base)
 {
     AutoRTFM::RecordOpenWrite(First, Last - First);
     return FN(First, Last, Value, Base);
 }
 
 template<typename T, std::to_chars_result FN(char*, char*, T, std::chars_format)>
-std::to_chars_result RTFM_ToChars(char* First, char* Last, T Value, std::chars_format Format)
+AUTORTFM_INTERNAL std::to_chars_result RTFM_ToChars(char* First, char* Last, T Value, std::chars_format Format)
 {
     AutoRTFM::RecordOpenWrite(First, Last - First);
     return FN(First, Last, Value, Format);
 }
 
 template<typename T, std::to_chars_result FN(char*, char*, T, std::chars_format, int Precision)>
-std::to_chars_result RTFM_ToChars(char* First, char* Last, T Value, std::chars_format Format, int Precision)
+AUTORTFM_INTERNAL std::to_chars_result RTFM_ToChars(char* First, char* Last, T Value, std::chars_format Format, int Precision)
 {
     AutoRTFM::RecordOpenWrite(First, Last - First);
     return FN(First, Last, Value, Format, Precision);
 }
 
 // FIXME: Does not currently support %n format specifiers.
-int RTFM_vsnprintf(char* Str, size_t Size, const char* Format, va_list ArgList) throw()
+AUTORTFM_INTERNAL int RTFM_vsnprintf(char* Str, size_t Size, const char* Format, va_list ArgList) throw()
 {
     ThrowIfFormatContainsPercentN(Format);
 
@@ -502,7 +502,7 @@ int RTFM_vsnprintf(char* Str, size_t Size, const char* Format, va_list ArgList)
 }
 
 // FIXME: Does not currently support %n format specifiers.
-int RTFM_vswprintf(wchar_t* Str, size_t Size, const wchar_t* Format, va_list ArgList)
+AUTORTFM_INTERNAL int RTFM_vswprintf(wchar_t* Str, size_t Size, const wchar_t* Format, va_list ArgList)
 {
     ThrowIfFormatContainsPercentN(Format);
 
@@ -533,7 +533,7 @@ int RTFM_vswprintf(wchar_t* Str, size_t Size, const wchar_t* Format, va_list Arg
 }
 
 // FIXME: Does not currently support %n format specifiers.
-int RTFM_swprintf(wchar_t* Buffer, size_t BufferCount, wchar_t const* Format, ...)
+AUTORTFM_INTERNAL int RTFM_swprintf(wchar_t* Buffer, size_t BufferCount, wchar_t const* Format, ...)
 {
     va_list ArgList;
 
@@ -545,7 +545,7 @@ int RTFM_swprintf(wchar_t* Buffer, size_t BufferCount, wchar_t const* Format, ..
 }
 
 // FIXME: Does not currently support %n format specifiers.
-int RTFM_snprintf(char* Str, size_t Size, const char* Format, ...) throw()
+AUTORTFM_INTERNAL int RTFM_snprintf(char* Str, size_t Size, const char* Format, ...) throw()
 {
     va_list ArgList;
 
@@ -557,7 +557,7 @@ int RTFM_snprintf(char* Str, size_t Size, const char* Format, ...) throw()
 }
 
 // FIXME: Does not currently support %n format specifiers.
-int RTFM_printf(const char* Format, ...)
+AUTORTFM_INTERNAL int RTFM_printf(const char* Format, ...)
 {
     ThrowIfFormatContainsPercentN(Format);
 
@@ -569,7 +569,7 @@ int RTFM_printf(const char* Format, ...)
 }
 
 // FIXME: Does not currently support %n format specifiers.
-int RTFM_wprintf(const wchar_t* Format, ...)
+AUTORTFM_INTERNAL int RTFM_wprintf(const wchar_t* Format, ...)
 {
     ThrowIfFormatContainsPercentN(Format);
 
@@ -580,7 +580,7 @@ int RTFM_wprintf(const wchar_t* Format, ...)
     return Result;
 }
 
-wchar_t* RTFM_wcscpy(wchar_t* Dst, const wchar_t* Src) throw()
+AUTORTFM_INTERNAL wchar_t* RTFM_wcscpy(wchar_t* Dst, const wchar_t* Src) throw()
 {
     const size_t SrcLen = wcslen(Src);
 
@@ -589,14 +589,14 @@ wchar_t* RTFM_wcscpy(wchar_t* Dst, const wchar_t* Src) throw()
     return wcscpy(Dst, Src);
 }
 
-wchar_t* RTFM_wcsncpy(wchar_t* Dst, const wchar_t* Src, size_t Count) throw()
+AUTORTFM_INTERNAL wchar_t* RTFM_wcsncpy(wchar_t* Dst, const wchar_t* Src, size_t Count) throw()
 {
 	FContext* Context = FContext::Get();
 	Context->RecordWrite(Dst, Count * sizeof(wchar_t));
 	return wcsncpy(Dst, Src, Count);
 }
 
-int RTFM_atexit(void(*Callback)(void)) throw()
+AUTORTFM_INTERNAL int RTFM_atexit(void(*Callback)(void)) throw()
 {
 	FContext* Context = FContext::Get();
 	Context->GetCurrentTransaction()->DeferUntilCommit([Callback]
diff --git a/Engine/Source/Runtime/AutoRTFM/Private/Transaction.h b/Engine/Source/Runtime/AutoRTFM/Private/Transaction.h
--- a/Engine/Source/Runtime/AutoRTFM/Private/Transaction.h
+++ b/Engine/Source/Runtime/AutoRTFM/Private/Transaction.h
@@ -19,7 +19,7 @@ namespace AutoRTFM
 {
 class FContext;
 
-class AUTORTFM_CORE_API FTransaction final
+class AUTORTFM_DISABLE FTransaction final
 {
 	friend class TIntrusivePool<FTransaction, 16>;
 
diff --git a/Engine/Source/Runtime/AutoRTFM/Private/Utils.h b/Engine/Source/Runtime/AutoRTFM/Private/Utils.h
--- a/Engine/Source/Runtime/AutoRTFM/Private/Utils.h
+++ b/Engine/Source/Runtime/AutoRTFM/Private/Utils.h
@@ -2,6 +2,7 @@
 
 #pragma once
 
+#include "AutoRTFMDefines.h"
 #if (defined(__AUTORTFM) && __AUTORTFM)
 
 #include "AutoRTFM.h"
@@ -22,7 +23,8 @@
 // declaration (usually in the header) and not a function implementation.
 // Annotated functions are callable from both the open and closed, without
 // trampolines or validation, and the function is always uninstrumented.
-#define AUTORTFM_CORE_API [[clang::autortfm_core_api]]
+// Note: This annotation will prevent inlining of annotated functions.
+#define AUTORTFM_INTERNAL [[clang::autortfm_core_api]]
 
 namespace AutoRTFM
 {
@@ -31,7 +33,7 @@ namespace AutoRTFM
 // based on ForTheRuntime::GetInternalAbortAction() and 
 // ForTheRuntime::GetEnsureOnInternalAbort().
 // If ProgramCounter is null, then the return address of the call will be used.
-__attribute__((__format__(__printf__, 4, 5)))
+AUTORTFM_DISABLE __attribute__((__format__(__printf__, 4, 5)))
 UE_AUTORTFM_API void ReportError(const char* File, int Line, void* ProgramCounter, const char* Format, ...);
 
 [[noreturn]] inline void InternalUnreachable(); //-V1082 Silence PVS [[noreturn]] false positive warning
diff --git a/Engine/Source/Runtime/AutoRTFM/Public/AutoRTFM.h b/Engine/Source/Runtime/AutoRTFM/Public/AutoRTFM.h
--- a/Engine/Source/Runtime/AutoRTFM/Public/AutoRTFM.h
+++ b/Engine/Source/Runtime/AutoRTFM/Public/AutoRTFM.h
@@ -569,13 +569,13 @@ auto AutoRTFMLookupInstrumentedFunctorInvoker(const FunctorType& Functor) -> voi
 // Tells if we are currently running in a transaction. This will return true in an open nest
 // (see `Open`). This function is handled specially in the compiler and will be constant folded
 // as true in closed code, or preserved as a function call in open code.
-UE_AUTORTFM_CRITICAL_INLINE bool IsTransactional() { return autortfm_is_transactional(); }
+UE_AUTORTFM_CRITICAL_INLINE_ALWAYS bool IsTransactional() { return autortfm_is_transactional(); }
 
 // Tells if we are currently running in the closed nest of a transaction. By default,
 // transactional code is in a closed nest; the only way to be in an open nest is to request it
 // via `Open`. This function is handled specially in the compiler and will be constant folded
 // as true in closed code, and false in open code.
-UE_AUTORTFM_CRITICAL_INLINE bool IsClosed() { return autortfm_is_closed(); }
+UE_AUTORTFM_CRITICAL_INLINE_ALWAYS bool IsClosed() { return autortfm_is_closed(); }
 
 // Tells us if we are currently committing or aborting a transaction. This will return true
 // in an on-abort or on-commit.
@@ -1075,7 +1075,7 @@ template<typename Type> UE_AUTORTFM_CRITICAL_INLINE void RecordOpenWriteNoMemory
 // Report that a unreachable codepath is being hit. Used to manually ban certain codepaths
 // from being transactionally safe.
 #if UE_AUTORTFM_ENABLED
-[[noreturn]] UE_AUTORTFM_CRITICAL_INLINE void Unreachable(const char* Message = nullptr)
+[[noreturn]] UE_AUTORTFM_CRITICAL_INLINE_ALWAYS void Unreachable(const char* Message = nullptr)
 {
 	autortfm_unreachable(Message);
 }
@@ -1087,7 +1087,7 @@ UE_AUTORTFM_CRITICAL_INLINE void Unreachable(const char* Message = nullptr)
 #endif
 
 // If we are running within a transaction, call `AutoRTFM::Unreachable`.
-UE_AUTORTFM_CRITICAL_INLINE void UnreachableIfTransactional(const char* Message = nullptr)
+UE_AUTORTFM_CRITICAL_INLINE_ALWAYS void UnreachableIfTransactional(const char* Message = nullptr)
 {
 	if (AutoRTFM::IsTransactional())
 	{
@@ -1096,7 +1096,7 @@ UE_AUTORTFM_CRITICAL_INLINE void UnreachableIfTransactional(const char* Message
 }
 
 // If we are running within a closed transaction, call `AutoRTFM::Unreachable`.
-UE_AUTORTFM_CRITICAL_INLINE void UnreachableIfClosed(const char* Message = nullptr)
+UE_AUTORTFM_CRITICAL_INLINE_ALWAYS void UnreachableIfClosed(const char* Message = nullptr)
 {
 	if (AutoRTFM::IsClosed())
 	{
diff --git a/Engine/Source/Runtime/AutoRTFM/Public/AutoRTFMDefines.h b/Engine/Source/Runtime/AutoRTFM/Public/AutoRTFMDefines.h
--- a/Engine/Source/Runtime/AutoRTFM/Public/AutoRTFMDefines.h
+++ b/Engine/Source/Runtime/AutoRTFM/Public/AutoRTFMDefines.h
@@ -88,18 +88,22 @@
 
 #if defined(_MSC_VER)
 #define UE_AUTORTFM_FORCEINLINE __forceinline
+#define UE_AUTORTFM_FORCEINLINE_ALWAYS __forceinline
 #define UE_AUTORTFM_FORCENOINLINE __declspec(noinline)
 #define UE_AUTORTFM_ASSUME(x) __assume(x)
 #elif defined(__clang__)
 #define UE_AUTORTFM_FORCEINLINE __attribute__((always_inline)) inline
+#define UE_AUTORTFM_FORCEINLINE_ALWAYS __attribute__((always_inline)) inline
 #define UE_AUTORTFM_FORCENOINLINE __attribute__((noinline))
 #define UE_AUTORTFM_ASSUME(x) __builtin_assume(x)
 #else
 #define UE_AUTORTFM_FORCEINLINE inline
+#define UE_AUTORTFM_FORCEINLINE_ALWAYS inline
 #define UE_AUTORTFM_FORCENOINLINE
 #define UE_AUTORTFM_ASSUME(x)
 #endif
 
+
 #if (defined(UE_BUILD_DEBUG) && UE_BUILD_DEBUG) || (defined(AUTORTFM_BUILD_DEBUG) && AUTORTFM_BUILD_DEBUG)
 // Force-inlining can make debugging glitchy. Disable this if we're running a debug build.
 #undef UE_AUTORTFM_CALLSITE_FORCEINLINE
@@ -125,6 +129,8 @@
 // This is a trick to ensure that there will never be a non-inlined version of these functions that the linker can decide to use
 #ifndef UE_HEADER_UNITS
 #define UE_AUTORTFM_CRITICAL_INLINE static UE_AUTORTFM_FORCEINLINE
+#define UE_AUTORTFM_CRITICAL_INLINE_ALWAYS static UE_AUTORTFM_FORCEINLINE_ALWAYS
 #else
 #define UE_AUTORTFM_CRITICAL_INLINE UE_AUTORTFM_FORCEINLINE // TODO: This needs to be revisited. we don't want bloated executables when modules are enabled
+#define UE_AUTORTFM_CRITICAL_INLINE_ALWAYS UE_AUTORTFM_FORCEINLINE_ALWAYS // TODO: This needs to be revisited. we don't want bloated executables when modules are enabled
 #endif
