[AutoRTFM] Add __is_autortfm_disabled(), AUTORTFM_ENABLE

__is_autortfm_disabled() is a new compiler traits function that evaluates to true if the function pointer or class type argument is AutoRTFM disabled.

The [[clang::autortfm_disable]] and [[clang::autortfm_core_api]] attributes have been replaced with [[clang::autortfm(disable)]] and [[clang::autortfm(internal)]].
This consolidates the mutually-exclusive attributes with a single parameterized attribute. This reduces compiler complexity in disallowing conflicting attributes, and simplifies the logic for determining a AutoRTFM "mode" for a function or method.

The new, unified attribute supports a second, optional condition which takes for form: [[clang::autortfm(mode, condition)]].
The condition parameter can be fed the result of  __is_autortfm_disabled() to propagate the AUTORTFM_DISABLE from callee to a templated function or class.

AUTORTFM_ENABLE / [[clang::autortfm(enable)]] is a new attribute which re-enables the intrumentation on a function that otherwise be disabled with AUTORTFM_DISABLE. The use case for this is cherry-picking methods that should be enabled on a class that is majority AUTORTFM_DISABLED.

[[clang::autortfm()]] attributes now implicitly propagate to child lexical scopes. For example: a lambda in a AUTORTFM_DISABLE'd function will also no longer have instrumentation generated for it.

Bulky, common LLVM clang logic has been moved out of the patch files and into AutoRTFMClang.[cpp|h].

The "proxy" .h and .cpp files that simply #include'd the out-of-llvm-tree files have been removed from the LLVM source tree. Instead, the CMake files directly reference the .cpp and .h files in the out-of-llvm-tree  AutoRTFM/Compiler directory. This removes the need to copy these files around, and helps improve code completion tools.

diff --git a/Engine/Source/Runtime/AutoRTFM/Private/Utils.h b/Engine/Source/Runtime/AutoRTFM/Private/Utils.h
--- a/Engine/Source/Runtime/AutoRTFM/Private/Utils.h
+++ b/Engine/Source/Runtime/AutoRTFM/Private/Utils.h
@@ -24,7 +24,7 @@
 // Annotated functions are callable from both the open and closed, without
 // trampolines or validation, and the function is always uninstrumented.
 // Note: This annotation will prevent inlining of annotated functions.
-#define AUTORTFM_INTERNAL [[clang::autortfm_core_api]]
+#define AUTORTFM_INTERNAL [[clang::autortfm(internal)]]
 
 namespace AutoRTFM
 {
diff --git a/Engine/Source/Runtime/AutoRTFM/Public/AutoRTFMDefines.h b/Engine/Source/Runtime/AutoRTFM/Public/AutoRTFMDefines.h
--- a/Engine/Source/Runtime/AutoRTFM/Public/AutoRTFMDefines.h
+++ b/Engine/Source/Runtime/AutoRTFM/Public/AutoRTFMDefines.h
@@ -44,7 +44,13 @@
 // declaration (usually in the header) and not a function implementation.
 // Annotated functions cannot be called from another closed function and
 // attempting to do so will result in a runtime failure.
-#define AUTORTFM_DISABLE [[clang::autortfm_disable]]
+#define AUTORTFM_DISABLE [[clang::autortfm(disable)]]
+
+// Annotation that can be applied to classes, methods or functions to re-enable
+// AutoRTFM instrumentation which would otherwise be disabled by AUTORTFM_DISABLE.
+// Useful for selectively enabling AutoRTFM on methods when a class is annotated
+// AUTORTFM_DISABLE.
+#define AUTORTFM_ENABLE [[clang::autortfm(enable)]]
 
 // Force the call statement to be inlined.
 #define UE_AUTORTFM_CALLSITE_FORCEINLINE [[clang::always_inline]]
