[AutoRTFM][StdLibSupport] Extern "C" and untemplate-ize all closed variants

Simplifies usage with the new AutoRTFM compiler open->closed mapping files.

jira SOL-8110, SOL-8242

diff --git a/Engine/Source/Runtime/AutoRTFM/Private/StdLibSupport.cpp b/Engine/Source/Runtime/AutoRTFM/Private/StdLibSupport.cpp
--- a/Engine/Source/Runtime/AutoRTFM/Private/StdLibSupport.cpp
+++ b/Engine/Source/Runtime/AutoRTFM/Private/StdLibSupport.cpp
@@ -79,6 +79,15 @@ private:
     FILE* const File;
 };
 
+template<typename T>
+inline void RecordOpenWriteIfNotNull(T* Ptr)
+{
+	if (Ptr != nullptr)
+	{
+		RecordOpenWrite(Ptr);
+	}
+}
+
 AUTORTFM_INTERNAL void ThrowErrorFormatContainsPercentN()
 {
     AUTORTFM_WARN("AutoRTFM does not support format strings containing '%%n'");
@@ -126,37 +135,79 @@ AUTORTFM_INTERNAL static void ThrowIfFormatContainsPercentN(const wchar_t* Forma
 
 #if AUTORTFM_PLATFORM_WINDOWS
 // vvvvvvvvvvvvvvvvvvvvvvvvvvvvvv AUTORTFM_PLATFORM_WINDOWS vvvvvvvvvvvvvvvvvvvvvvvvvvvvv
-template<typename RETURN, typename CHAR, RETURN FN(const CHAR*, CHAR**, _locale_t)>
-AUTORTFM_INTERNAL RETURN RTFM_StringToFloat(const CHAR* String, CHAR** EndPtr, _locale_t Locale)
+extern "C" AUTORTFM_INTERNAL long long RTFM__strtoi64(const char* String, char** EndPtr, int Radix)
 {
-    if (nullptr != EndPtr)
-    {
-        AutoRTFM::RecordOpenWrite(EndPtr);
-    }
-    return FN(String, EndPtr, Locale);
+	RecordOpenWriteIfNotNull(EndPtr);
+	return _strtoi64(String, EndPtr, Radix);
 }
 
-template<typename RETURN, typename CHAR, RETURN FN(const CHAR*, CHAR**, int)>
-AUTORTFM_INTERNAL RETURN RTFM_StringToIntNoLocale(const CHAR* String, CHAR** EndPtr, int Radix)
+extern "C" AUTORTFM_INTERNAL long long RTFM__wcstoi64(const wchar_t* String, wchar_t** EndPtr, int Radix)
 {
-    if (nullptr != EndPtr)
-    {
-        AutoRTFM::RecordOpenWrite(EndPtr);
-    }
-    return FN(String, EndPtr, Radix);
+	RecordOpenWriteIfNotNull(EndPtr);
+	return _wcstoi64(String, EndPtr, Radix);
 }
 
-template<typename RETURN, typename CHAR, RETURN FN(const CHAR*, CHAR**, int, _locale_t)>
-AUTORTFM_INTERNAL RETURN RTFM_StringToInt(const CHAR* String, CHAR** EndPtr, int Radix, _locale_t Locale)
+extern "C" AUTORTFM_INTERNAL unsigned long long RTFM__wcstoui64(const wchar_t* String, wchar_t** EndPtr, int Radix)
 {
-    if (nullptr != EndPtr)
-    {
-        AutoRTFM::RecordOpenWrite(EndPtr);
-    }
-    return FN(String, EndPtr, Radix, Locale);
+	RecordOpenWriteIfNotNull(EndPtr);
+	return _wcstoui64(String, EndPtr, Radix);
+}
+
+extern "C" AUTORTFM_INTERNAL double RTFM__wcstod_l(const wchar_t* String, wchar_t** EndPtr, _locale_t Locale)
+{
+	RecordOpenWriteIfNotNull(EndPtr);
+	return _wcstod_l(String, EndPtr, Locale);
+}
+
+extern "C" AUTORTFM_INTERNAL float RTFM__wcstof_l(const wchar_t* String, wchar_t** EndPtr, _locale_t Locale)
+{
+	RecordOpenWriteIfNotNull(EndPtr);
+	return _wcstof_l(String, EndPtr, Locale);
+}
+
+extern "C" AUTORTFM_INTERNAL long double RTFM__wcstold_l(const wchar_t* String, wchar_t** EndPtr, _locale_t Locale)
+{
+	RecordOpenWriteIfNotNull(EndPtr);
+	return _wcstold_l(String, EndPtr, Locale);
+}
+
+extern "C" AUTORTFM_INTERNAL long RTFM__wcstol_l(const wchar_t* String, wchar_t** EndPtr, int Radix, _locale_t Locale)
+{
+	RecordOpenWriteIfNotNull(EndPtr);
+	return _wcstol_l(String, EndPtr, Radix, Locale);
+}
+
+extern "C" AUTORTFM_INTERNAL long long RTFM__wcstoll_l(const wchar_t* String, wchar_t** EndPtr, int Radix, _locale_t Locale)
+{
+	RecordOpenWriteIfNotNull(EndPtr);
+	return _wcstoll_l(String, EndPtr, Radix, Locale);
+}
+
+extern "C" AUTORTFM_INTERNAL unsigned long RTFM__wcstoul_l(const wchar_t* String, wchar_t** EndPtr, int Radix, _locale_t Locale)
+{
+	RecordOpenWriteIfNotNull(EndPtr);
+	return _wcstoul_l(String, EndPtr, Radix, Locale);
+}
+
+extern "C" AUTORTFM_INTERNAL unsigned long long RTFM__wcstoull_l(const wchar_t* String, wchar_t** EndPtr, int Radix, _locale_t Locale)
+{
+	RecordOpenWriteIfNotNull(EndPtr);
+	return _wcstoull_l(String, EndPtr, Radix, Locale);
+}
+
+extern "C" AUTORTFM_INTERNAL long long RTFM__wcstoi64_l(const wchar_t* String, wchar_t** EndPtr, int Radix, _locale_t Locale)
+{
+	RecordOpenWriteIfNotNull(EndPtr);
+	return _wcstoi64_l(String, EndPtr, Radix, Locale);
+}
+
+extern "C" AUTORTFM_INTERNAL unsigned long long RTFM__wcstoui64_l(const wchar_t* String, wchar_t** EndPtr, int Radix, _locale_t Locale)
+{
+	RecordOpenWriteIfNotNull(EndPtr);
+	return _wcstoui64_l(String, EndPtr, Radix, Locale);
 }
 
-AUTORTFM_INTERNAL FILE* RTFM___acrt_iob_func(unsigned Index)
+extern "C" AUTORTFM_INTERNAL FILE* RTFM___acrt_iob_func(unsigned Index)
 {
     switch (Index)
     {
@@ -174,7 +225,7 @@ AUTORTFM_INTERNAL FILE* RTFM___acrt_iob_func(unsigned Index)
 }
 
 // FIXME: Does not currently support %n format specifiers.
-AUTORTFM_INTERNAL int RTFM___stdio_common_vfprintf(
+extern "C" AUTORTFM_INTERNAL int RTFM___stdio_common_vfprintf(
         unsigned __int64 Options,
         FILE*            Stream,
         char const*      Format,
@@ -187,7 +238,7 @@ AUTORTFM_INTERNAL int RTFM___stdio_common_vfprintf(
 }
 
 // FIXME: Does not currently support %n format specifiers.
-AUTORTFM_INTERNAL int RTFM___stdio_common_vsprintf(
+extern "C" AUTORTFM_INTERNAL int RTFM___stdio_common_vsprintf(
         unsigned __int64 Options,
         char*            Buffer,
         size_t           BufferCount,
@@ -216,7 +267,7 @@ AUTORTFM_INTERNAL int RTFM___stdio_common_vsprintf(
 }
 
 // FIXME: Does not currently support %n format specifiers.
-AUTORTFM_INTERNAL int RTFM___stdio_common_vswprintf(
+extern "C" AUTORTFM_INTERNAL int RTFM___stdio_common_vswprintf(
         unsigned __int64 Options,
         wchar_t*         Buffer,
         size_t           BufferCount,
@@ -245,7 +296,7 @@ AUTORTFM_INTERNAL int RTFM___stdio_common_vswprintf(
 }
 
 // FIXME: Does not currently support %n format specifiers.
-AUTORTFM_INTERNAL int RTFM___stdio_common_vfwprintf(
+extern "C" AUTORTFM_INTERNAL int RTFM___stdio_common_vfwprintf(
         unsigned __int64 Options,
         FILE*            Stream,
         wchar_t const*   Format,
@@ -257,7 +308,7 @@ AUTORTFM_INTERNAL int RTFM___stdio_common_vfwprintf(
     return __stdio_common_vfwprintf(Options, Stream, Format, Locale, ArgList);
 }
 
-AUTORTFM_INTERNAL BOOL RTFM_TlsSetValue(DWORD dwTlsIndex, LPVOID lpTlsValue)
+extern "C" AUTORTFM_INTERNAL BOOL RTFM_TlsSetValue(DWORD dwTlsIndex, LPVOID lpTlsValue)
 {
 	LPVOID OriginalValue = TlsGetValue(dwTlsIndex);
 
@@ -269,7 +320,7 @@ AUTORTFM_INTERNAL BOOL RTFM_TlsSetValue(DWORD dwTlsIndex, LPVOID lpTlsValue)
 	return TlsSetValue(dwTlsIndex, lpTlsValue);
 }
 
-BOOL RTFM_FlsSetValue(DWORD dwFlsIndex, LPVOID lpFlsValue)
+extern "C" BOOL RTFM_FlsSetValue(DWORD dwFlsIndex, LPVOID lpFlsValue)
 {
 	LPVOID OriginalValue = FlsGetValue(dwFlsIndex);
 
@@ -290,7 +341,7 @@ extern "C" size_t _ZNSt3__112__next_primeEm(size_t N) __attribute__((weak));
 
 #if AUTORTFM_PLATFORM_LINUX
 // vvvvvvvvvvvvvvvvvvvvvvvvvvvvvv AUTORTFM_PLATFORM_LINUX vvvvvvvvvvvvvvvvvvvvvvvvvvvvvv
-AUTORTFM_INTERNAL int RTFM_stat(const char* Path, struct stat* StatBuf) throw()
+extern "C" AUTORTFM_INTERNAL int RTFM_stat(const char* Path, struct stat* StatBuf) throw()
 {
 	FContext* Context = FContext::Get();
 	Context->RecordWrite(StatBuf, sizeof(*StatBuf));
@@ -298,7 +349,7 @@ AUTORTFM_INTERNAL int RTFM_stat(const char* Path, struct stat* StatBuf) throw()
 	return stat(Path, StatBuf);
 }
 
-AUTORTFM_INTERNAL int RTFM_fstat(int Fd, struct stat* StatBuf) throw()
+extern "C" AUTORTFM_INTERNAL int RTFM_fstat(int Fd, struct stat* StatBuf) throw()
 {
 	FContext* Context = FContext::Get();
 	Context->RecordWrite(StatBuf, sizeof(*StatBuf));
@@ -306,7 +357,7 @@ AUTORTFM_INTERNAL int RTFM_fstat(int Fd, struct stat* StatBuf) throw()
 	return fstat(Fd, StatBuf);
 }
 
-AUTORTFM_INTERNAL int RTFM___xstat(int Ver, const char* Path, struct stat* StatBuf) throw()
+extern "C" AUTORTFM_INTERNAL int RTFM___xstat(int Ver, const char* Path, struct stat* StatBuf) throw()
 {
 	FContext* Context = FContext::Get();
 	Context->RecordWrite(StatBuf, sizeof(*StatBuf));
@@ -314,32 +365,45 @@ AUTORTFM_INTERNAL int RTFM___xstat(int Ver, const char* Path, struct stat* StatB
 	return __xstat(Ver, Path, StatBuf);
 }
 
-AUTORTFM_INTERNAL int RTFM___fxstat(int Ver, int Fd, struct stat* StatBuf) throw()
+extern "C" AUTORTFM_INTERNAL int RTFM___fxstat(int Ver, int Fd, struct stat* StatBuf) throw()
 {
 	FContext* Context = FContext::Get();
 	Context->RecordWrite(StatBuf, sizeof(*StatBuf));
 
 	return __fxstat(Ver, Fd, StatBuf);
 }
+
+extern "C" AUTORTFM_INTERNAL float RTFM_strtof32(const char* String, char** EndPtr) throw()
+{
+	RecordOpenWriteIfNotNull(EndPtr);
+	return strtof32(String, EndPtr);
+}
+
+extern "C" AUTORTFM_INTERNAL double RTFM_strtof64(const char* String, char** EndPtr) throw()
+{
+	RecordOpenWriteIfNotNull(EndPtr);
+	return strtof64(String, EndPtr);
+}
+
 // ^^^^^^^^^^^^^^^^^^^^^^^^^^^^^^ AUTORTFM_PLATFORM_LINUX ^^^^^^^^^^^^^^^^^^^^^^^^^^^^^^
 #endif
 
-AUTORTFM_INTERNAL void* RTFM_memcpy(void* Dst, const void* Src, size_t Size) throw()
+extern "C" AUTORTFM_INTERNAL void* RTFM_memcpy(void* Dst, const void* Src, size_t Size) throw()
 {
     return Memcpy(Dst, Src, Size, FContext::Get());
 }
 
-AUTORTFM_INTERNAL void* RTFM_memmove(void* Dst, const void* Src, size_t Size) throw()
+extern "C" AUTORTFM_INTERNAL void* RTFM_memmove(void* Dst, const void* Src, size_t Size) throw()
 {
     return Memmove(Dst, Src, Size, FContext::Get());
 }
 
-AUTORTFM_INTERNAL void* RTFM_memset(void* Dst, int Value, size_t Size) throw()
+extern "C" AUTORTFM_INTERNAL void* RTFM_memset(void* Dst, int Value, size_t Size) throw()
 {
     return Memset(Dst, Value, Size, FContext::Get());
 }
 
-AUTORTFM_INTERNAL void* RTFM_malloc(size_t Size) throw()
+extern "C" AUTORTFM_INTERNAL void* RTFM_malloc(size_t Size) throw()
 {
     void* Result = malloc(Size);
 	FContext* Context = FContext::Get();
@@ -351,7 +415,7 @@ AUTORTFM_INTERNAL void* RTFM_malloc(size_t Size) throw()
     return Result;
 }
 
-AUTORTFM_INTERNAL void* RTFM_calloc(size_t Count, size_t Size) throw()
+extern "C" AUTORTFM_INTERNAL void* RTFM_calloc(size_t Count, size_t Size) throw()
 {
     void* Result = calloc(Count, Size);
 	FContext* Context = FContext::Get();
@@ -363,7 +427,7 @@ AUTORTFM_INTERNAL void* RTFM_calloc(size_t Count, size_t Size) throw()
     return Result;
 }
 
-AUTORTFM_INTERNAL void RTFM_free(void* Ptr) throw()
+extern "C" AUTORTFM_INTERNAL void RTFM_free(void* Ptr) throw()
 {
     if (Ptr)
     {
@@ -375,7 +439,7 @@ AUTORTFM_INTERNAL void RTFM_free(void* Ptr) throw()
     }
 }
 
-AUTORTFM_INTERNAL void* RTFM_realloc(void* Ptr, size_t Size) throw()
+extern "C" AUTORTFM_INTERNAL void* RTFM_realloc(void* Ptr, size_t Size) throw()
 {
     void* NewObject = RTFM_malloc(Size);
     if (Ptr)
@@ -394,7 +458,7 @@ AUTORTFM_INTERNAL void* RTFM_realloc(void* Ptr, size_t Size) throw()
     return NewObject;
 }
 
-AUTORTFM_INTERNAL char* RTFM_strcpy(char* const Dst, const char* const Src) throw()
+extern "C" AUTORTFM_INTERNAL char* RTFM_strcpy(char* const Dst, const char* const Src) throw()
 {
     const size_t SrcLen = strlen(Src);
 
@@ -403,14 +467,14 @@ AUTORTFM_INTERNAL char* RTFM_strcpy(char* const Dst, const char* const Src) thro
     return strcpy(Dst, Src);
 }
 
-AUTORTFM_INTERNAL char* RTFM_strncpy(char* const Dst, const char* const Src, const size_t Num) throw()
+extern "C" AUTORTFM_INTERNAL char* RTFM_strncpy(char* const Dst, const char* const Src, const size_t Num) throw()
 {
 	FContext* Context = FContext::Get();
     Context->RecordWrite(Dst, Num);
     return strncpy(Dst, Src, Num);
 }
 
-AUTORTFM_INTERNAL char* RTFM_strcat(char* const Dst, const char* const Src) throw()
+extern "C" AUTORTFM_INTERNAL char* RTFM_strcat(char* const Dst, const char* const Src) throw()
 {
     const size_t DstLen = strlen(Dst);
     const size_t SrcLen = strlen(Src);
@@ -420,7 +484,7 @@ AUTORTFM_INTERNAL char* RTFM_strcat(char* const Dst, const char* const Src) thro
     return strcat(Dst, Src);
 }
 
-AUTORTFM_INTERNAL char* RTFM_strncat(char* const Dst, const char* const Src, const size_t Num) throw()
+extern "C" AUTORTFM_INTERNAL char* RTFM_strncat(char* const Dst, const char* const Src, const size_t Num) throw()
 {
     const size_t DstLen = strlen(Dst);
 
@@ -429,56 +493,200 @@ AUTORTFM_INTERNAL char* RTFM_strncat(char* const Dst, const char* const Src, con
     return strncat(Dst, Src, Num);
 }
 
-template<typename RETURN, typename CHAR, RETURN FN(const CHAR*, CHAR**)>
-AUTORTFM_INTERNAL RETURN RTFM_StringToFloat(const CHAR* String, CHAR** EndPtr) throw()
+extern "C" AUTORTFM_INTERNAL long RTFM_strtol(const char* String, char** EndPtr, int Radix) throw()
+{ 
+    RecordOpenWriteIfNotNull(EndPtr);
+	return strtol(String, EndPtr, Radix); 
+}
+
+extern "C" AUTORTFM_INTERNAL long long RTFM_strtoll(const char* String, char** EndPtr, int Radix) throw()
+{ 
+    RecordOpenWriteIfNotNull(EndPtr);
+	return strtoll(String, EndPtr, Radix); 
+}
+
+extern "C" AUTORTFM_INTERNAL unsigned long RTFM_strtoul(const char* String, char** EndPtr, int Radix) throw()
+{ 
+    RecordOpenWriteIfNotNull(EndPtr);
+	return strtoul(String, EndPtr, Radix); 
+}
+
+extern "C" AUTORTFM_INTERNAL unsigned long long RTFM_strtoull(const char* String, char** EndPtr, int Radix) throw()
+{ 
+    RecordOpenWriteIfNotNull(EndPtr);
+	return strtoull(String, EndPtr, Radix); 
+}
+
+extern "C" AUTORTFM_INTERNAL float RTFM_strtof(const char* String, char** EndPtr) throw()
+{ 
+    RecordOpenWriteIfNotNull(EndPtr);
+	return strtof(String, EndPtr); 
+}
+
+extern "C" AUTORTFM_INTERNAL double RTFM_strtod(const char* String, char** EndPtr) throw()
+{ 
+    RecordOpenWriteIfNotNull(EndPtr);
+	return strtod(String, EndPtr); 
+}
+
+extern "C" AUTORTFM_INTERNAL double RTFM_wcstod(const wchar_t* String, wchar_t** EndPtr) throw()
+{ 
+    RecordOpenWriteIfNotNull(EndPtr);
+	return wcstod(String, EndPtr); 
+}
+
+extern "C" AUTORTFM_INTERNAL float RTFM_wcstof(const wchar_t* String, wchar_t** EndPtr) throw()
+{ 
+    RecordOpenWriteIfNotNull(EndPtr);
+	return wcstof(String, EndPtr); 
+}
+
+extern "C" AUTORTFM_INTERNAL long double RTFM_wcstold(const wchar_t* String, wchar_t** EndPtr) throw()
+{ 
+    RecordOpenWriteIfNotNull(EndPtr);
+	return wcstold(String, EndPtr); 
+}
+
+extern "C" AUTORTFM_INTERNAL long RTFM_wcstol(const wchar_t* String, wchar_t** EndPtr, int Radix) throw()
+{ 
+    RecordOpenWriteIfNotNull(EndPtr);
+	return wcstol(String, EndPtr, Radix); 
+}
+
+extern "C" AUTORTFM_INTERNAL long long RTFM_wcstoll(const wchar_t* String, wchar_t** EndPtr, int Radix) throw()
+{ 
+    RecordOpenWriteIfNotNull(EndPtr);
+	return wcstoll(String, EndPtr, Radix); 
+}
+
+extern "C" AUTORTFM_INTERNAL unsigned long RTFM_wcstoul(const wchar_t* String, wchar_t** EndPtr, int Radix) throw()
+{ 
+    RecordOpenWriteIfNotNull(EndPtr);
+	return wcstoul(String, EndPtr, Radix); 
+}
+
+extern "C" AUTORTFM_INTERNAL unsigned long long RTFM_wcstoull(const wchar_t* String, wchar_t** EndPtr, int Radix) throw()
+{ 
+    RecordOpenWriteIfNotNull(EndPtr);
+	return wcstoull(String, EndPtr, Radix); 
+}
+
+extern "C" AUTORTFM_INTERNAL std::to_chars_result RTFM__ToChars_i(char* First, char* Last, int Value, int Base)
 {
-    if (nullptr != EndPtr)
-    {
-        AutoRTFM::RecordOpenWrite(EndPtr);
-    }
-    return FN(String, EndPtr);
+	RecordOpenWrite(First, Last - First);
+	return std::to_chars(First, Last, Value, Base);
 }
 
-template<typename RETURN, typename CHAR, RETURN FN(const CHAR*, CHAR**, int)>
-AUTORTFM_INTERNAL RETURN RTFM_StringToInt(const CHAR* String, CHAR** EndPtr, int Radix) throw()
+extern "C" AUTORTFM_INTERNAL std::to_chars_result RTFM__ToChars_u(char* First, char* Last, unsigned int Value, int Base)
 {
-    if (nullptr != EndPtr)
-    {
-        AutoRTFM::RecordOpenWrite(EndPtr);
-    }
-    return FN(String, EndPtr, Radix);
+	RecordOpenWrite(First, Last - First);
+	return std::to_chars(First, Last, Value, Base);
+}
+
+extern "C" AUTORTFM_INTERNAL std::to_chars_result RTFM__ToChars_i8(char* First, char* Last, int8_t Value, int Base)
+{
+	RecordOpenWrite(First, Last - First);
+	return std::to_chars(First, Last, Value, Base);
+}
+
+extern "C" AUTORTFM_INTERNAL std::to_chars_result RTFM__ToChars_u8(char* First, char* Last, uint8_t Value, int Base)
+{
+	RecordOpenWrite(First, Last - First);
+	return std::to_chars(First, Last, Value, Base);
 }
 
-template<typename T, std::to_chars_result FN(char*, char*, T)>
-AUTORTFM_INTERNAL std::to_chars_result RTFM_ToChars(char* First, char* Last, T Value)
+extern "C" AUTORTFM_INTERNAL std::to_chars_result RTFM__ToChars_i16(char* First, char* Last, int16_t Value, int Base)
 {
-    AutoRTFM::RecordOpenWrite(First, Last - First);
-    return FN(First, Last, Value);
+	RecordOpenWrite(First, Last - First);
+	return std::to_chars(First, Last, Value, Base);
 }
 
-template<typename T, std::to_chars_result FN(char*, char*, T, int)>
-AUTORTFM_INTERNAL std::to_chars_result RTFM_ToChars(char* First, char* Last, T Value, int Base)
+extern "C" AUTORTFM_INTERNAL std::to_chars_result RTFM__ToChars_u16(char* First, char* Last, uint16_t Value, int Base)
 {
-    AutoRTFM::RecordOpenWrite(First, Last - First);
-    return FN(First, Last, Value, Base);
+	RecordOpenWrite(First, Last - First);
+	return std::to_chars(First, Last, Value, Base);
 }
 
-template<typename T, std::to_chars_result FN(char*, char*, T, std::chars_format)>
-AUTORTFM_INTERNAL std::to_chars_result RTFM_ToChars(char* First, char* Last, T Value, std::chars_format Format)
+extern "C" AUTORTFM_INTERNAL std::to_chars_result RTFM__ToChars_i32(char* First, char* Last, int32_t Value, int Base)
 {
-    AutoRTFM::RecordOpenWrite(First, Last - First);
-    return FN(First, Last, Value, Format);
+	RecordOpenWrite(First, Last - First);
+	return std::to_chars(First, Last, Value, Base);
 }
 
-template<typename T, std::to_chars_result FN(char*, char*, T, std::chars_format, int Precision)>
-AUTORTFM_INTERNAL std::to_chars_result RTFM_ToChars(char* First, char* Last, T Value, std::chars_format Format, int Precision)
+extern "C" AUTORTFM_INTERNAL std::to_chars_result RTFM__ToChars_u32(char* First, char* Last, uint32_t Value, int Base)
 {
-    AutoRTFM::RecordOpenWrite(First, Last - First);
-    return FN(First, Last, Value, Format, Precision);
+	RecordOpenWrite(First, Last - First);
+	return std::to_chars(First, Last, Value, Base);
+}
+
+extern "C" AUTORTFM_INTERNAL std::to_chars_result RTFM__ToChars_i64(char* First, char* Last, int64_t Value, int Base)
+{
+	RecordOpenWrite(First, Last - First);
+	return std::to_chars(First, Last, Value, Base);
+}
+
+extern "C" AUTORTFM_INTERNAL std::to_chars_result RTFM__ToChars_u64(char* First, char* Last, uint64_t Value, int Base)
+{
+	RecordOpenWrite(First, Last - First);
+	return std::to_chars(First, Last, Value, Base);
+}
+
+extern "C" AUTORTFM_INTERNAL std::to_chars_result RTFM__ToChars_f(char* First, char* Last, float Value)
+{
+	RecordOpenWrite(First, Last - First);
+	return std::to_chars(First, Last, Value);
+}
+
+extern "C" AUTORTFM_INTERNAL std::to_chars_result RTFM__ToChars_f_F(char* First, char* Last, float Value, std::chars_format Format)
+{
+	RecordOpenWrite(First, Last - First);
+	return std::to_chars(First, Last, Value, Format);
+}
+
+extern "C" AUTORTFM_INTERNAL std::to_chars_result RTFM__ToChars_f_FP(char* First, char* Last, float Value, std::chars_format Format, int Precision)
+{
+	RecordOpenWrite(First, Last - First);
+	return std::to_chars(First, Last, Value, Format, Precision);
+}
+
+extern "C" AUTORTFM_INTERNAL std::to_chars_result RTFM__ToChars_d(char* First, char* Last, double Value)
+{
+	RecordOpenWrite(First, Last - First);
+	return std::to_chars(First, Last, Value);
+}
+
+extern "C" AUTORTFM_INTERNAL std::to_chars_result RTFM__ToChars_d_F(char* First, char* Last, double Value, std::chars_format Format)
+{
+	RecordOpenWrite(First, Last - First);
+	return std::to_chars(First, Last, Value, Format);
+}
+
+extern "C" AUTORTFM_INTERNAL std::to_chars_result RTFM__ToChars_d_FP(char* First, char* Last, double Value, std::chars_format Format, int Precision)
+{
+	RecordOpenWrite(First, Last - First);
+	return std::to_chars(First, Last, Value, Format, Precision);
+}
+
+extern "C" AUTORTFM_INTERNAL std::to_chars_result RTFM__ToChars_ld(char* First, char* Last, long double Value)
+{
+	RecordOpenWrite(First, Last - First);
+	return std::to_chars(First, Last, Value);
+}
+
+extern "C" AUTORTFM_INTERNAL std::to_chars_result RTFM__ToChars_ld_F(char* First, char* Last, long double Value, std::chars_format Format)
+{
+	RecordOpenWrite(First, Last - First);
+	return std::to_chars(First, Last, Value, Format);
+}
+
+extern "C" AUTORTFM_INTERNAL std::to_chars_result RTFM__ToChars_ld_FP(char* First, char* Last, long double Value, std::chars_format Format, int Precision)
+{
+	RecordOpenWrite(First, Last - First);
+	return std::to_chars(First, Last, Value, Format, Precision);
 }
 
 // FIXME: Does not currently support %n format specifiers.
-AUTORTFM_INTERNAL int RTFM_vsnprintf(char* Str, size_t Size, const char* Format, va_list ArgList) throw()
+extern "C" AUTORTFM_INTERNAL int RTFM_vsnprintf(char* Str, size_t Size, const char* Format, va_list ArgList) throw()
 {
     ThrowIfFormatContainsPercentN(Format);
 
@@ -502,7 +710,7 @@ AUTORTFM_INTERNAL int RTFM_vsnprintf(char* Str, size_t Size, const char* Format,
 }
 
 // FIXME: Does not currently support %n format specifiers.
-AUTORTFM_INTERNAL int RTFM_vswprintf(wchar_t* Str, size_t Size, const wchar_t* Format, va_list ArgList)
+extern "C" AUTORTFM_INTERNAL int RTFM_vswprintf(wchar_t* Str, size_t Size, const wchar_t* Format, va_list ArgList)
 {
     ThrowIfFormatContainsPercentN(Format);
 
@@ -533,7 +741,7 @@ AUTORTFM_INTERNAL int RTFM_vswprintf(wchar_t* Str, size_t Size, const wchar_t* F
 }
 
 // FIXME: Does not currently support %n format specifiers.
-AUTORTFM_INTERNAL int RTFM_swprintf(wchar_t* Buffer, size_t BufferCount, wchar_t const* Format, ...)
+extern "C" AUTORTFM_INTERNAL int RTFM_swprintf(wchar_t* Buffer, size_t BufferCount, wchar_t const* Format, ...)
 {
     va_list ArgList;
 
@@ -545,7 +753,7 @@ AUTORTFM_INTERNAL int RTFM_swprintf(wchar_t* Buffer, size_t BufferCount, wchar_t
 }
 
 // FIXME: Does not currently support %n format specifiers.
-AUTORTFM_INTERNAL int RTFM_snprintf(char* Str, size_t Size, const char* Format, ...) throw()
+extern "C" AUTORTFM_INTERNAL int RTFM_snprintf(char* Str, size_t Size, const char* Format, ...) throw()
 {
     va_list ArgList;
 
@@ -557,7 +765,7 @@ AUTORTFM_INTERNAL int RTFM_snprintf(char* Str, size_t Size, const char* Format,
 }
 
 // FIXME: Does not currently support %n format specifiers.
-AUTORTFM_INTERNAL int RTFM_printf(const char* Format, ...)
+extern "C" AUTORTFM_INTERNAL int RTFM_printf(const char* Format, ...)
 {
     ThrowIfFormatContainsPercentN(Format);
 
@@ -569,7 +777,7 @@ AUTORTFM_INTERNAL int RTFM_printf(const char* Format, ...)
 }
 
 // FIXME: Does not currently support %n format specifiers.
-AUTORTFM_INTERNAL int RTFM_wprintf(const wchar_t* Format, ...)
+extern "C" AUTORTFM_INTERNAL int RTFM_wprintf(const wchar_t* Format, ...)
 {
     ThrowIfFormatContainsPercentN(Format);
 
@@ -580,7 +788,7 @@ AUTORTFM_INTERNAL int RTFM_wprintf(const wchar_t* Format, ...)
     return Result;
 }
 
-AUTORTFM_INTERNAL wchar_t* RTFM_wcscpy(wchar_t* Dst, const wchar_t* Src) throw()
+extern "C" AUTORTFM_INTERNAL wchar_t* RTFM_wcscpy(wchar_t* Dst, const wchar_t* Src) throw()
 {
     const size_t SrcLen = wcslen(Src);
 
@@ -589,14 +797,14 @@ AUTORTFM_INTERNAL wchar_t* RTFM_wcscpy(wchar_t* Dst, const wchar_t* Src) throw()
     return wcscpy(Dst, Src);
 }
 
-AUTORTFM_INTERNAL wchar_t* RTFM_wcsncpy(wchar_t* Dst, const wchar_t* Src, size_t Count) throw()
+extern "C" AUTORTFM_INTERNAL wchar_t* RTFM_wcsncpy(wchar_t* Dst, const wchar_t* Src, size_t Count) throw()
 {
 	FContext* Context = FContext::Get();
 	Context->RecordWrite(Dst, Count * sizeof(wchar_t));
 	return wcsncpy(Dst, Src, Count);
 }
 
-AUTORTFM_INTERNAL int RTFM_atexit(void(*Callback)(void)) throw()
+extern "C" AUTORTFM_INTERNAL int RTFM_atexit(void(*Callback)(void)) throw()
 {
 	FContext* Context = FContext::Get();
 	Context->GetCurrentTransaction()->DeferUntilCommit([Callback]
@@ -685,90 +893,109 @@ extern "C"
 	void _ZdaPvSt11align_val_tRKSt9nothrow_t(void*, std::align_val_t, const std::nothrow_t&);
 	void _ZdlPvmSt11align_val_t(void*, unsigned long, std::align_val_t);
 	void _ZdaPvmSt11align_val_t(void*, unsigned long, std::align_val_t);
+} // extern "C"
 
-	void* RTFM_Znwm(unsigned long Size)
-	{
-		return RTFM_ClosedNew(_Znwm, _ZdlPv, Size);
-	}
-	void* RTFM_Znam(unsigned long Size)
-	{
-		return RTFM_ClosedNew(_Znam, _ZdaPv, Size);
-	}
-	void* RTFM_ZnwmRKSt9nothrow_t(unsigned long Size, const std::nothrow_t& NoThrow)
-	{
-		return RTFM_ClosedNew(_ZnwmRKSt9nothrow_t, _ZdlPvRKSt9nothrow_t, Size, NoThrow);
-	}
-	void* RTFM_ZnamRKSt9nothrow_t(unsigned long Size, const std::nothrow_t& NoThrow)
-	{
-		return RTFM_ClosedNew(_ZnamRKSt9nothrow_t, _ZdaPvRKSt9nothrow_t, Size, NoThrow);
-	}
-	void* RTFM_ZnwmSt11align_val_t(unsigned long Size, std::align_val_t Align)
-	{
-		return RTFM_ClosedNew(_ZnwmSt11align_val_t, _ZdlPvSt11align_val_t, Size, Align);
-	}
-	void* RTFM_ZnamSt11align_val_t(unsigned long Size, std::align_val_t Align)
-	{
-		return RTFM_ClosedNew(_ZnamSt11align_val_t, _ZdaPvSt11align_val_t, Size, Align);
-	}
-	void* RTFM_ZnwmSt11align_val_tRKSt9nothrow_t(unsigned long Size, std::align_val_t Align, const std::nothrow_t& NoThrow)
-	{
-		return RTFM_ClosedNew(_ZnwmSt11align_val_tRKSt9nothrow_t, _ZdlPvSt11align_val_tRKSt9nothrow_t, Size, Align, NoThrow);
-	}
-	void* RTFM_ZnamSt11align_val_tRKSt9nothrow_t(unsigned long Size, std::align_val_t Align, const std::nothrow_t& NoThrow)
-	{
-		return RTFM_ClosedNew(_ZnamSt11align_val_tRKSt9nothrow_t, _ZdaPvSt11align_val_tRKSt9nothrow_t, Size, Align, NoThrow);
-	}
+extern "C" AUTORTFM_INTERNAL void* RTFM_Znwm(unsigned long Size)
+{
+	return RTFM_ClosedNew(_Znwm, _ZdlPv, Size);
+}
 
-	void RTFM_ZdlPv(void* Pointer)
-	{
-		RTFM_ClosedDelete(_ZdlPv, Pointer);
-	}
-	void RTFM_ZdaPv(void* Pointer)
-	{
-		RTFM_ClosedDelete(_ZdaPv, Pointer);
-	}
-	void RTFM_ZdlPvRKSt9nothrow_t(void* Pointer, const std::nothrow_t& NoThrow)
-	{
-		RTFM_ClosedDelete(_ZdlPvRKSt9nothrow_t, Pointer, NoThrow);
-	}
-	void RTFM_ZdaPvRKSt9nothrow_t(void* Pointer, const std::nothrow_t& NoThrow)
-	{
-		RTFM_ClosedDelete(_ZdaPvRKSt9nothrow_t, Pointer, NoThrow);
-	}
-	void RTFM_ZdlPvm(void* Pointer, unsigned long Size)
-	{
-		RTFM_ClosedDelete(_ZdlPvm, Pointer, Size);
-	}
-	void RTFM_ZdaPvm(void* Pointer, unsigned long Size)
-	{
-		RTFM_ClosedDelete(_ZdaPvm, Pointer, Size);
-	}
-	void RTFM_ZdlPvSt11align_val_t(void* Pointer, std::align_val_t Align)
-	{
-		RTFM_ClosedDelete(_ZdlPvSt11align_val_t, Pointer, Align);
-	}
-	void RTFM_ZdaPvSt11align_val_t(void* Pointer, std::align_val_t Align)
-	{
-		RTFM_ClosedDelete(_ZdaPvSt11align_val_t, Pointer, Align);
-	}
-	void RTFM_ZdlPvSt11align_val_tRKSt9nothrow_t(void* Pointer, std::align_val_t Align, const std::nothrow_t& NoThrow)
-	{
-		RTFM_ClosedDelete(_ZdlPvSt11align_val_tRKSt9nothrow_t, Pointer, Align, NoThrow);
-	}
-	void RTFM_ZdaPvSt11align_val_tRKSt9nothrow_t(void* Pointer, std::align_val_t Align, const std::nothrow_t& NoThrow)
-	{
-		RTFM_ClosedDelete(_ZdaPvSt11align_val_tRKSt9nothrow_t, Pointer, Align, NoThrow);
-	}
-	void RTFM_ZdlPvmSt11align_val_t(void* Pointer, unsigned long Size, std::align_val_t Align)
-	{
-		RTFM_ClosedDelete(_ZdlPvmSt11align_val_t, Pointer, Size, Align);
-	}
-	void RTFM_ZdaPvmSt11align_val_t(void* Pointer, unsigned long Size, std::align_val_t Align)
-	{
-		RTFM_ClosedDelete(_ZdaPvmSt11align_val_t, Pointer, Size, Align);
-	}
+extern "C" AUTORTFM_INTERNAL void* RTFM_Znam(unsigned long Size)
+{
+	return RTFM_ClosedNew(_Znam, _ZdaPv, Size);
 }
-#endif
+
+extern "C" AUTORTFM_INTERNAL void* RTFM_ZnwmRKSt9nothrow_t(unsigned long Size, const std::nothrow_t& NoThrow)
+{
+	return RTFM_ClosedNew(_ZnwmRKSt9nothrow_t, _ZdlPvRKSt9nothrow_t, Size, NoThrow);
+}
+
+extern "C" AUTORTFM_INTERNAL void* RTFM_ZnamRKSt9nothrow_t(unsigned long Size, const std::nothrow_t& NoThrow)
+{
+	return RTFM_ClosedNew(_ZnamRKSt9nothrow_t, _ZdaPvRKSt9nothrow_t, Size, NoThrow);
+}
+
+extern "C" AUTORTFM_INTERNAL void* RTFM_ZnwmSt11align_val_t(unsigned long Size, std::align_val_t Align)
+{
+	return RTFM_ClosedNew(_ZnwmSt11align_val_t, _ZdlPvSt11align_val_t, Size, Align);
+}
+
+extern "C" AUTORTFM_INTERNAL void* RTFM_ZnamSt11align_val_t(unsigned long Size, std::align_val_t Align)
+{
+	return RTFM_ClosedNew(_ZnamSt11align_val_t, _ZdaPvSt11align_val_t, Size, Align);
+}
+
+extern "C" AUTORTFM_INTERNAL void* RTFM_ZnwmSt11align_val_tRKSt9nothrow_t(unsigned long Size, std::align_val_t Align, const std::nothrow_t& NoThrow)
+{
+	return RTFM_ClosedNew(_ZnwmSt11align_val_tRKSt9nothrow_t, _ZdlPvSt11align_val_tRKSt9nothrow_t, Size, Align, NoThrow);
+}
+
+extern "C" AUTORTFM_INTERNAL void* RTFM_ZnamSt11align_val_tRKSt9nothrow_t(unsigned long Size, std::align_val_t Align, const std::nothrow_t& NoThrow)
+{
+	return RTFM_ClosedNew(_ZnamSt11align_val_tRKSt9nothrow_t, _ZdaPvSt11align_val_tRKSt9nothrow_t, Size, Align, NoThrow);
+}
+
+extern "C" AUTORTFM_INTERNAL void RTFM_ZdlPv(void* Pointer)
+{
+	RTFM_ClosedDelete(_ZdlPv, Pointer);
+}
+
+extern "C" AUTORTFM_INTERNAL void RTFM_ZdaPv(void* Pointer)
+{
+	RTFM_ClosedDelete(_ZdaPv, Pointer);
+}
+
+extern "C" AUTORTFM_INTERNAL void RTFM_ZdlPvRKSt9nothrow_t(void* Pointer, const std::nothrow_t& NoThrow)
+{
+	RTFM_ClosedDelete(_ZdlPvRKSt9nothrow_t, Pointer, NoThrow);
+}
+
+extern "C" AUTORTFM_INTERNAL void RTFM_ZdaPvRKSt9nothrow_t(void* Pointer, const std::nothrow_t& NoThrow)
+{
+	RTFM_ClosedDelete(_ZdaPvRKSt9nothrow_t, Pointer, NoThrow);
+}
+
+extern "C" AUTORTFM_INTERNAL void RTFM_ZdlPvm(void* Pointer, unsigned long Size)
+{
+	RTFM_ClosedDelete(_ZdlPvm, Pointer, Size);
+}
+
+extern "C" AUTORTFM_INTERNAL void RTFM_ZdaPvm(void* Pointer, unsigned long Size)
+{
+	RTFM_ClosedDelete(_ZdaPvm, Pointer, Size);
+}
+
+extern "C" AUTORTFM_INTERNAL void RTFM_ZdlPvSt11align_val_t(void* Pointer, std::align_val_t Align)
+{
+	RTFM_ClosedDelete(_ZdlPvSt11align_val_t, Pointer, Align);
+}
+
+extern "C" AUTORTFM_INTERNAL void RTFM_ZdaPvSt11align_val_t(void* Pointer, std::align_val_t Align)
+{
+	RTFM_ClosedDelete(_ZdaPvSt11align_val_t, Pointer, Align);
+}
+
+extern "C" AUTORTFM_INTERNAL void RTFM_ZdlPvSt11align_val_tRKSt9nothrow_t(void* Pointer, std::align_val_t Align, const std::nothrow_t& NoThrow)
+{
+	RTFM_ClosedDelete(_ZdlPvSt11align_val_tRKSt9nothrow_t, Pointer, Align, NoThrow);
+}
+
+extern "C" AUTORTFM_INTERNAL void RTFM_ZdaPvSt11align_val_tRKSt9nothrow_t(void* Pointer, std::align_val_t Align, const std::nothrow_t& NoThrow)
+{
+	RTFM_ClosedDelete(_ZdaPvSt11align_val_tRKSt9nothrow_t, Pointer, Align, NoThrow);
+}
+
+extern "C" AUTORTFM_INTERNAL void RTFM_ZdlPvmSt11align_val_t(void* Pointer, unsigned long Size, std::align_val_t Align)
+{
+	RTFM_ClosedDelete(_ZdlPvmSt11align_val_t, Pointer, Size, Align);
+}
+
+extern "C" AUTORTFM_INTERNAL void RTFM_ZdaPvmSt11align_val_t(void* Pointer, unsigned long Size, std::align_val_t Align)
+{
+	RTFM_ClosedDelete(_ZdaPvmSt11align_val_t, Pointer, Size, Align);
+}
+
+#endif // FORCE_ANSI_ALLOCATOR
 
 UE_AUTORTFM_REGISTER_OPEN_TO_CLOSED_FUNCTIONS(
 #if AUTORTFM_ASAN_ENABLED
@@ -778,18 +1005,18 @@ UE_AUTORTFM_REGISTER_OPEN_TO_CLOSED_FUNCTIONS(
 #endif
 
 #if AUTORTFM_PLATFORM_WINDOWS
-	UE_AUTORTFM_MAP_OPEN_TO_CLOSED(_strtoi64, (RTFM_StringToIntNoLocale<long long, char, _strtoi64>)),
-	UE_AUTORTFM_MAP_OPEN_TO_CLOSED(_wcstoi64, (RTFM_StringToInt<long long, wchar_t, _wcstoi64>)),
-	UE_AUTORTFM_MAP_OPEN_TO_CLOSED(_wcstoui64, (RTFM_StringToInt<unsigned long long, wchar_t, _wcstoui64>)),
-	UE_AUTORTFM_MAP_OPEN_TO_CLOSED(_wcstod_l, (RTFM_StringToFloat<double, wchar_t, _wcstod_l>)),
-	UE_AUTORTFM_MAP_OPEN_TO_CLOSED(_wcstof_l, (RTFM_StringToFloat<float, wchar_t, _wcstof_l>)),
-	UE_AUTORTFM_MAP_OPEN_TO_CLOSED(_wcstold_l, (RTFM_StringToFloat<long double, wchar_t, _wcstold_l>)),
-	UE_AUTORTFM_MAP_OPEN_TO_CLOSED(_wcstol_l, (RTFM_StringToInt<long, wchar_t, _wcstol_l>)),
-	UE_AUTORTFM_MAP_OPEN_TO_CLOSED(_wcstoll_l, (RTFM_StringToInt<long long, wchar_t, _wcstoll_l>)),
-	UE_AUTORTFM_MAP_OPEN_TO_CLOSED(_wcstoul_l, (RTFM_StringToInt<unsigned long, wchar_t, _wcstoul_l>)),
-	UE_AUTORTFM_MAP_OPEN_TO_CLOSED(_wcstoull_l, (RTFM_StringToInt<unsigned long long, wchar_t, _wcstoull_l>)),
-	UE_AUTORTFM_MAP_OPEN_TO_CLOSED(_wcstoi64_l, (RTFM_StringToInt<long long, wchar_t, _wcstoi64_l>)),
-	UE_AUTORTFM_MAP_OPEN_TO_CLOSED(_wcstoui64_l, (RTFM_StringToInt<unsigned long long, wchar_t, _wcstoui64_l>)),
+	UE_AUTORTFM_MAP_OPEN_TO_CLOSED(_strtoi64, RTFM__strtoi64),
+	UE_AUTORTFM_MAP_OPEN_TO_CLOSED(_wcstoi64, RTFM__wcstoi64),
+	UE_AUTORTFM_MAP_OPEN_TO_CLOSED(_wcstoui64, RTFM__wcstoui64),
+	UE_AUTORTFM_MAP_OPEN_TO_CLOSED(_wcstod_l, RTFM__wcstod_l),
+	UE_AUTORTFM_MAP_OPEN_TO_CLOSED(_wcstof_l, RTFM__wcstof_l),
+	UE_AUTORTFM_MAP_OPEN_TO_CLOSED(_wcstold_l, RTFM__wcstold_l),
+	UE_AUTORTFM_MAP_OPEN_TO_CLOSED(_wcstol_l, RTFM__wcstol_l),
+	UE_AUTORTFM_MAP_OPEN_TO_CLOSED(_wcstoll_l, RTFM__wcstoll_l),
+	UE_AUTORTFM_MAP_OPEN_TO_CLOSED(_wcstoul_l, RTFM__wcstoul_l),
+	UE_AUTORTFM_MAP_OPEN_TO_CLOSED(_wcstoull_l, RTFM__wcstoull_l),
+	UE_AUTORTFM_MAP_OPEN_TO_CLOSED(_wcstoi64_l, RTFM__wcstoi64_l),
+	UE_AUTORTFM_MAP_OPEN_TO_CLOSED(_wcstoui64_l, RTFM__wcstoui64_l),
 	UE_AUTORTFM_MAP_OPEN_TO_SELF(_wtof),
 	UE_AUTORTFM_MAP_OPEN_TO_SELF(_wtof),
 	UE_AUTORTFM_MAP_OPEN_TO_SELF(_wtol),
@@ -852,8 +1079,8 @@ UE_AUTORTFM_REGISTER_OPEN_TO_CLOSED_FUNCTIONS(
 	UE_AUTORTFM_MAP_OPEN_TO_SELF(bcmp),
 	UE_AUTORTFM_MAP_OPEN_TO_SELF(pthread_getspecific),
 	UE_AUTORTFM_MAP_OPEN_TO_SELF(pthread_self),
-	UE_AUTORTFM_MAP_OPEN_TO_CLOSED(strtof32, (RTFM_StringToFloat<float, char, strtof32>)),
-	UE_AUTORTFM_MAP_OPEN_TO_CLOSED(strtof64, (RTFM_StringToFloat<double, char, strtof64>)),
+	UE_AUTORTFM_MAP_OPEN_TO_CLOSED(strtof32, RTFM_strtof32),
+	UE_AUTORTFM_MAP_OPEN_TO_CLOSED(strtof64, RTFM_strtof64),
 	UE_AUTORTFM_MAP_OPEN_TO_CLOSED(stat, RTFM_stat),
 	UE_AUTORTFM_MAP_OPEN_TO_CLOSED(fstat, RTFM_fstat),
 	UE_AUTORTFM_MAP_OPEN_TO_CLOSED(__xstat, RTFM___xstat),
@@ -920,95 +1147,76 @@ UE_AUTORTFM_REGISTER_OPEN_TO_CLOSED_FUNCTIONS(
 	UE_AUTORTFM_MAP_OPEN_TO_SELF_OVERLOADED(const wchar_t*(const wchar_t*, wchar_t), wcsrchr),
 	UE_AUTORTFM_MAP_OPEN_TO_SELF(wcscmp),
 	UE_AUTORTFM_MAP_OPEN_TO_SELF(wcslen),
-	UE_AUTORTFM_MAP_OPEN_TO_CLOSED(strtol, (RTFM_StringToInt<long, char, strtol>)),
-	UE_AUTORTFM_MAP_OPEN_TO_CLOSED(strtoll, (RTFM_StringToInt<long long, char, strtoll>)),
-	UE_AUTORTFM_MAP_OPEN_TO_CLOSED(strtoul, (RTFM_StringToInt<unsigned long, char, strtoul>)),
-	UE_AUTORTFM_MAP_OPEN_TO_CLOSED(strtoull, (RTFM_StringToInt<unsigned long long, char, strtoull>)),
-	UE_AUTORTFM_MAP_OPEN_TO_CLOSED(strtof, (RTFM_StringToFloat<float, char, strtof>)),
-	UE_AUTORTFM_MAP_OPEN_TO_CLOSED(strtod, (RTFM_StringToFloat<double, char, strtod>)),
-	UE_AUTORTFM_MAP_OPEN_TO_CLOSED(wcstod, (RTFM_StringToFloat<double, wchar_t, wcstod>)),
-	UE_AUTORTFM_MAP_OPEN_TO_CLOSED(wcstof, (RTFM_StringToFloat<float, wchar_t, wcstof>)),
-	UE_AUTORTFM_MAP_OPEN_TO_CLOSED(wcstold, (RTFM_StringToFloat<long double, wchar_t, wcstold>)),
-	UE_AUTORTFM_MAP_OPEN_TO_CLOSED(wcstol, (RTFM_StringToInt<long, wchar_t, wcstol>)),
-	UE_AUTORTFM_MAP_OPEN_TO_CLOSED(wcstoll, (RTFM_StringToInt<long long, wchar_t, wcstoll>)),
-	UE_AUTORTFM_MAP_OPEN_TO_CLOSED(wcstoul, (RTFM_StringToInt<unsigned long, wchar_t, wcstoul>)),
-	UE_AUTORTFM_MAP_OPEN_TO_CLOSED(wcstoull, (RTFM_StringToInt<unsigned long long, wchar_t, wcstoull>)),
+	UE_AUTORTFM_MAP_OPEN_TO_CLOSED(strtol, RTFM_strtol),
+	UE_AUTORTFM_MAP_OPEN_TO_CLOSED(strtoll, RTFM_strtoll),
+	UE_AUTORTFM_MAP_OPEN_TO_CLOSED(strtoul, RTFM_strtoul),
+	UE_AUTORTFM_MAP_OPEN_TO_CLOSED(strtoull, RTFM_strtoull),
+	UE_AUTORTFM_MAP_OPEN_TO_CLOSED(strtof, RTFM_strtof),
+	UE_AUTORTFM_MAP_OPEN_TO_CLOSED(strtod, RTFM_strtod),
+	UE_AUTORTFM_MAP_OPEN_TO_CLOSED(wcstod, RTFM_wcstod),
+	UE_AUTORTFM_MAP_OPEN_TO_CLOSED(wcstof, RTFM_wcstof),
+	UE_AUTORTFM_MAP_OPEN_TO_CLOSED(wcstold, RTFM_wcstold),
+	UE_AUTORTFM_MAP_OPEN_TO_CLOSED(wcstol, RTFM_wcstol),
+	UE_AUTORTFM_MAP_OPEN_TO_CLOSED(wcstoll, RTFM_wcstoll),
+	UE_AUTORTFM_MAP_OPEN_TO_CLOSED(wcstoul, RTFM_wcstoul),
+	UE_AUTORTFM_MAP_OPEN_TO_CLOSED(wcstoull, RTFM_wcstoull),
 	UE_AUTORTFM_MAP_OPEN_TO_CLOSED_OVERLOADED(
 		std::to_chars_result(char* First, char* Last, int Value, int Base),
-		std::to_chars,
-		(RTFM_ToChars<int, std::to_chars>)),
+		std::to_chars, RTFM__ToChars_i),
 	UE_AUTORTFM_MAP_OPEN_TO_CLOSED_OVERLOADED(
 		std::to_chars_result(char* First, char* Last, unsigned int Value, int Base),
-		std::to_chars,
-		(RTFM_ToChars<unsigned int, std::to_chars>)),
+		std::to_chars, RTFM__ToChars_u),
 	UE_AUTORTFM_MAP_OPEN_TO_CLOSED_OVERLOADED(
 		std::to_chars_result(char* First, char* Last, int8_t Value, int Base),
-		std::to_chars,
-		(RTFM_ToChars<int8_t, std::to_chars>)),
+		std::to_chars, RTFM__ToChars_i8),
 	UE_AUTORTFM_MAP_OPEN_TO_CLOSED_OVERLOADED(
 		std::to_chars_result(char* First, char* Last, uint8_t Value, int Base),
-		std::to_chars,
-		(RTFM_ToChars<uint8_t, std::to_chars>)),
+		std::to_chars, RTFM__ToChars_u8),
 	UE_AUTORTFM_MAP_OPEN_TO_CLOSED_OVERLOADED(
 		std::to_chars_result(char* First, char* Last, int16_t Value, int Base),
-		std::to_chars,
-		(RTFM_ToChars<int16_t, std::to_chars>)),
+		std::to_chars, RTFM__ToChars_i16),
 	UE_AUTORTFM_MAP_OPEN_TO_CLOSED_OVERLOADED(
 		std::to_chars_result(char* First, char* Last, uint16_t Value, int Base),
-		std::to_chars,
-		(RTFM_ToChars<uint16_t, std::to_chars>)),
+		std::to_chars, RTFM__ToChars_u16),
 	UE_AUTORTFM_MAP_OPEN_TO_CLOSED_OVERLOADED(
 		std::to_chars_result(char* First, char* Last, int32_t Value, int Base),
-		std::to_chars,
-		(RTFM_ToChars<int32_t, std::to_chars>)),
+		std::to_chars, RTFM__ToChars_i32),
 	UE_AUTORTFM_MAP_OPEN_TO_CLOSED_OVERLOADED(
 		std::to_chars_result(char* First, char* Last, uint32_t Value, int Base),
-		std::to_chars,
-		(RTFM_ToChars<uint32_t, std::to_chars>)),
+		std::to_chars, RTFM__ToChars_u32),
 	UE_AUTORTFM_MAP_OPEN_TO_CLOSED_OVERLOADED(
 		std::to_chars_result(char* First, char* Last, int64_t Value, int Base),
-		std::to_chars,
-		(RTFM_ToChars<int64_t, std::to_chars>)),
+		std::to_chars, RTFM__ToChars_i64),
 	UE_AUTORTFM_MAP_OPEN_TO_CLOSED_OVERLOADED(
 		std::to_chars_result(char* First, char* Last, uint64_t Value, int Base),
-		std::to_chars,
-		(RTFM_ToChars<uint64_t, std::to_chars>)),
+		std::to_chars, RTFM__ToChars_u64),
 	UE_AUTORTFM_MAP_OPEN_TO_CLOSED_OVERLOADED(
 		std::to_chars_result(char* First, char* Last, float Value),
-		std::to_chars,
-		(RTFM_ToChars<float, std::to_chars>)),
+		std::to_chars, RTFM__ToChars_f),
 	UE_AUTORTFM_MAP_OPEN_TO_CLOSED_OVERLOADED(
 		std::to_chars_result(char* First, char* Last, float Value, std::chars_format Format),
-		std::to_chars,
-		(RTFM_ToChars<float, std::to_chars>)),
+		std::to_chars, RTFM__ToChars_f_F),
 	UE_AUTORTFM_MAP_OPEN_TO_CLOSED_OVERLOADED(
 		std::to_chars_result(char* First, char* Last, float Value, std::chars_format Format, int Precision),
-		std::to_chars,
-		(RTFM_ToChars<float, std::to_chars>)),
+		std::to_chars, RTFM__ToChars_f_FP),
 	UE_AUTORTFM_MAP_OPEN_TO_CLOSED_OVERLOADED(
 		std::to_chars_result(char* First, char* Last, double Value),
-		std::to_chars,
-		(RTFM_ToChars<double, std::to_chars>)),
+		std::to_chars, RTFM__ToChars_d),
 	UE_AUTORTFM_MAP_OPEN_TO_CLOSED_OVERLOADED(
 		std::to_chars_result(char* First, char* Last, double Value, std::chars_format Format),
-		std::to_chars,
-		(RTFM_ToChars<double, std::to_chars>)),
+		std::to_chars, RTFM__ToChars_d_F),
 	UE_AUTORTFM_MAP_OPEN_TO_CLOSED_OVERLOADED(
 		std::to_chars_result(char* First, char* Last, double Value, std::chars_format Format, int Precision),
-		std::to_chars,
-		(RTFM_ToChars<double, std::to_chars>)),
+		std::to_chars, RTFM__ToChars_d_FP),
 	UE_AUTORTFM_MAP_OPEN_TO_CLOSED_OVERLOADED(
 		std::to_chars_result(char* First, char* Last, long double Value),
-		std::to_chars,
-		(RTFM_ToChars<long double, std::to_chars>)),
+		std::to_chars, RTFM__ToChars_ld),
 	UE_AUTORTFM_MAP_OPEN_TO_CLOSED_OVERLOADED(
 		std::to_chars_result(char* First, char* Last, long double Value, std::chars_format Format),
-		std::to_chars,
-		(RTFM_ToChars<long double, std::to_chars>)),
+		std::to_chars, RTFM__ToChars_ld_F),
 	UE_AUTORTFM_MAP_OPEN_TO_CLOSED_OVERLOADED(
 		std::to_chars_result(char* First, char* Last, long double Value, std::chars_format Format, int Precision),
-		std::to_chars,
-		(RTFM_ToChars<long double, std::to_chars>)),
+		std::to_chars, RTFM__ToChars_ld_FP),
 
 	UE_AUTORTFM_MAP_OPEN_TO_SELF(iswupper),
 	UE_AUTORTFM_MAP_OPEN_TO_SELF(iswlower),
