[AutoRTFM] Export all the stdlib closed functions

Replace the condition for the new / delete mappings from FORCE_ANSI_ALLOCATOR to AUTORTFM_PLATFORM_LINUX. There should be no harm in registering these when the ANSI allocator is disabled. These will not work on Windows, as it uses a completely different mangling.

diff --git a/Engine/Source/Runtime/AutoRTFM/Private/StdLibSupport.cpp b/Engine/Source/Runtime/AutoRTFM/Private/StdLibSupport.cpp
--- a/Engine/Source/Runtime/AutoRTFM/Private/StdLibSupport.cpp
+++ b/Engine/Source/Runtime/AutoRTFM/Private/StdLibSupport.cpp
@@ -133,81 +133,83 @@ AUTORTFM_INTERNAL static void ThrowIfFormatContainsPercentN(const wchar_t* Forma
     }
 }
 
+} // anonymous namespace
+
 #if AUTORTFM_PLATFORM_WINDOWS
 // vvvvvvvvvvvvvvvvvvvvvvvvvvvvvv AUTORTFM_PLATFORM_WINDOWS vvvvvvvvvvvvvvvvvvvvvvvvvvvvv
-extern "C" AUTORTFM_INTERNAL long long RTFM__strtoi64(const char* String, char** EndPtr, int Radix)
+extern "C" AUTORTFM_INTERNAL UE_AUTORTFM_API long long RTFM__strtoi64(const char* String, char** EndPtr, int Radix)
 {
 	RecordOpenWriteIfNotNull(EndPtr);
 	return _strtoi64(String, EndPtr, Radix);
 }
 
-extern "C" AUTORTFM_INTERNAL long long RTFM__wcstoi64(const wchar_t* String, wchar_t** EndPtr, int Radix)
+extern "C" AUTORTFM_INTERNAL UE_AUTORTFM_API long long RTFM__wcstoi64(const wchar_t* String, wchar_t** EndPtr, int Radix)
 {
 	RecordOpenWriteIfNotNull(EndPtr);
 	return _wcstoi64(String, EndPtr, Radix);
 }
 
-extern "C" AUTORTFM_INTERNAL unsigned long long RTFM__wcstoui64(const wchar_t* String, wchar_t** EndPtr, int Radix)
+extern "C" AUTORTFM_INTERNAL UE_AUTORTFM_API unsigned long long RTFM__wcstoui64(const wchar_t* String, wchar_t** EndPtr, int Radix)
 {
 	RecordOpenWriteIfNotNull(EndPtr);
 	return _wcstoui64(String, EndPtr, Radix);
 }
 
-extern "C" AUTORTFM_INTERNAL double RTFM__wcstod_l(const wchar_t* String, wchar_t** EndPtr, _locale_t Locale)
+extern "C" AUTORTFM_INTERNAL UE_AUTORTFM_API double RTFM__wcstod_l(const wchar_t* String, wchar_t** EndPtr, _locale_t Locale)
 {
 	RecordOpenWriteIfNotNull(EndPtr);
 	return _wcstod_l(String, EndPtr, Locale);
 }
 
-extern "C" AUTORTFM_INTERNAL float RTFM__wcstof_l(const wchar_t* String, wchar_t** EndPtr, _locale_t Locale)
+extern "C" AUTORTFM_INTERNAL UE_AUTORTFM_API float RTFM__wcstof_l(const wchar_t* String, wchar_t** EndPtr, _locale_t Locale)
 {
 	RecordOpenWriteIfNotNull(EndPtr);
 	return _wcstof_l(String, EndPtr, Locale);
 }
 
-extern "C" AUTORTFM_INTERNAL long double RTFM__wcstold_l(const wchar_t* String, wchar_t** EndPtr, _locale_t Locale)
+extern "C" AUTORTFM_INTERNAL UE_AUTORTFM_API long double RTFM__wcstold_l(const wchar_t* String, wchar_t** EndPtr, _locale_t Locale)
 {
 	RecordOpenWriteIfNotNull(EndPtr);
 	return _wcstold_l(String, EndPtr, Locale);
 }
 
-extern "C" AUTORTFM_INTERNAL long RTFM__wcstol_l(const wchar_t* String, wchar_t** EndPtr, int Radix, _locale_t Locale)
+extern "C" AUTORTFM_INTERNAL UE_AUTORTFM_API long RTFM__wcstol_l(const wchar_t* String, wchar_t** EndPtr, int Radix, _locale_t Locale)
 {
 	RecordOpenWriteIfNotNull(EndPtr);
 	return _wcstol_l(String, EndPtr, Radix, Locale);
 }
 
-extern "C" AUTORTFM_INTERNAL long long RTFM__wcstoll_l(const wchar_t* String, wchar_t** EndPtr, int Radix, _locale_t Locale)
+extern "C" AUTORTFM_INTERNAL UE_AUTORTFM_API long long RTFM__wcstoll_l(const wchar_t* String, wchar_t** EndPtr, int Radix, _locale_t Locale)
 {
 	RecordOpenWriteIfNotNull(EndPtr);
 	return _wcstoll_l(String, EndPtr, Radix, Locale);
 }
 
-extern "C" AUTORTFM_INTERNAL unsigned long RTFM__wcstoul_l(const wchar_t* String, wchar_t** EndPtr, int Radix, _locale_t Locale)
+extern "C" AUTORTFM_INTERNAL UE_AUTORTFM_API unsigned long RTFM__wcstoul_l(const wchar_t* String, wchar_t** EndPtr, int Radix, _locale_t Locale)
 {
 	RecordOpenWriteIfNotNull(EndPtr);
 	return _wcstoul_l(String, EndPtr, Radix, Locale);
 }
 
-extern "C" AUTORTFM_INTERNAL unsigned long long RTFM__wcstoull_l(const wchar_t* String, wchar_t** EndPtr, int Radix, _locale_t Locale)
+extern "C" AUTORTFM_INTERNAL UE_AUTORTFM_API unsigned long long RTFM__wcstoull_l(const wchar_t* String, wchar_t** EndPtr, int Radix, _locale_t Locale)
 {
 	RecordOpenWriteIfNotNull(EndPtr);
 	return _wcstoull_l(String, EndPtr, Radix, Locale);
 }
 
-extern "C" AUTORTFM_INTERNAL long long RTFM__wcstoi64_l(const wchar_t* String, wchar_t** EndPtr, int Radix, _locale_t Locale)
+extern "C" AUTORTFM_INTERNAL UE_AUTORTFM_API long long RTFM__wcstoi64_l(const wchar_t* String, wchar_t** EndPtr, int Radix, _locale_t Locale)
 {
 	RecordOpenWriteIfNotNull(EndPtr);
 	return _wcstoi64_l(String, EndPtr, Radix, Locale);
 }
 
-extern "C" AUTORTFM_INTERNAL unsigned long long RTFM__wcstoui64_l(const wchar_t* String, wchar_t** EndPtr, int Radix, _locale_t Locale)
+extern "C" AUTORTFM_INTERNAL UE_AUTORTFM_API unsigned long long RTFM__wcstoui64_l(const wchar_t* String, wchar_t** EndPtr, int Radix, _locale_t Locale)
 {
 	RecordOpenWriteIfNotNull(EndPtr);
 	return _wcstoui64_l(String, EndPtr, Radix, Locale);
 }
 
-extern "C" AUTORTFM_INTERNAL FILE* RTFM___acrt_iob_func(unsigned Index)
+extern "C" AUTORTFM_INTERNAL UE_AUTORTFM_API FILE* RTFM___acrt_iob_func(unsigned Index)
 {
     switch (Index)
     {
@@ -225,7 +227,7 @@ extern "C" AUTORTFM_INTERNAL FILE* RTFM___acrt_iob_func(unsigned Index)
 }
 
 // FIXME: Does not currently support %n format specifiers.
-extern "C" AUTORTFM_INTERNAL int RTFM___stdio_common_vfprintf(
+extern "C" AUTORTFM_INTERNAL UE_AUTORTFM_API int RTFM___stdio_common_vfprintf(
         unsigned __int64 Options,
         FILE*            Stream,
         char const*      Format,
@@ -238,7 +240,7 @@ extern "C" AUTORTFM_INTERNAL int RTFM___stdio_common_vfprintf(
 }
 
 // FIXME: Does not currently support %n format specifiers.
-extern "C" AUTORTFM_INTERNAL int RTFM___stdio_common_vsprintf(
+extern "C" AUTORTFM_INTERNAL UE_AUTORTFM_API int RTFM___stdio_common_vsprintf(
         unsigned __int64 Options,
         char*            Buffer,
         size_t           BufferCount,
@@ -267,7 +269,7 @@ extern "C" AUTORTFM_INTERNAL int RTFM___stdio_common_vsprintf(
 }
 
 // FIXME: Does not currently support %n format specifiers.
-extern "C" AUTORTFM_INTERNAL int RTFM___stdio_common_vswprintf(
+extern "C" AUTORTFM_INTERNAL UE_AUTORTFM_API int RTFM___stdio_common_vswprintf(
         unsigned __int64 Options,
         wchar_t*         Buffer,
         size_t           BufferCount,
@@ -296,7 +298,7 @@ extern "C" AUTORTFM_INTERNAL int RTFM___stdio_common_vswprintf(
 }
 
 // FIXME: Does not currently support %n format specifiers.
-extern "C" AUTORTFM_INTERNAL int RTFM___stdio_common_vfwprintf(
+extern "C" AUTORTFM_INTERNAL UE_AUTORTFM_API int RTFM___stdio_common_vfwprintf(
         unsigned __int64 Options,
         FILE*            Stream,
         wchar_t const*   Format,
@@ -308,7 +310,7 @@ extern "C" AUTORTFM_INTERNAL int RTFM___stdio_common_vfwprintf(
     return __stdio_common_vfwprintf(Options, Stream, Format, Locale, ArgList);
 }
 
-extern "C" AUTORTFM_INTERNAL BOOL RTFM_TlsSetValue(DWORD dwTlsIndex, LPVOID lpTlsValue)
+extern "C" AUTORTFM_INTERNAL UE_AUTORTFM_API BOOL RTFM_TlsSetValue(DWORD dwTlsIndex, LPVOID lpTlsValue)
 {
 	LPVOID OriginalValue = TlsGetValue(dwTlsIndex);
 
@@ -320,7 +322,7 @@ extern "C" AUTORTFM_INTERNAL BOOL RTFM_TlsSetValue(DWORD dwTlsIndex, LPVOID lpTl
 	return TlsSetValue(dwTlsIndex, lpTlsValue);
 }
 
-extern "C" BOOL RTFM_FlsSetValue(DWORD dwFlsIndex, LPVOID lpFlsValue)
+extern "C" AUTORTFM_INTERNAL UE_AUTORTFM_API BOOL RTFM_FlsSetValue(DWORD dwFlsIndex, LPVOID lpFlsValue)
 {
 	LPVOID OriginalValue = FlsGetValue(dwFlsIndex);
 
@@ -341,7 +343,7 @@ extern "C" size_t _ZNSt3__112__next_primeEm(size_t N) __attribute__((weak));
 
 #if AUTORTFM_PLATFORM_LINUX
 // vvvvvvvvvvvvvvvvvvvvvvvvvvvvvv AUTORTFM_PLATFORM_LINUX vvvvvvvvvvvvvvvvvvvvvvvvvvvvvv
-extern "C" AUTORTFM_INTERNAL int RTFM_stat(const char* Path, struct stat* StatBuf) throw()
+extern "C" AUTORTFM_INTERNAL UE_AUTORTFM_API int RTFM_stat(const char* Path, struct stat* StatBuf) throw()
 {
 	FContext* Context = FContext::Get();
 	Context->RecordWrite(StatBuf, sizeof(*StatBuf));
@@ -349,7 +351,7 @@ extern "C" AUTORTFM_INTERNAL int RTFM_stat(const char* Path, struct stat* StatBu
 	return stat(Path, StatBuf);
 }
 
-extern "C" AUTORTFM_INTERNAL int RTFM_fstat(int Fd, struct stat* StatBuf) throw()
+extern "C" AUTORTFM_INTERNAL UE_AUTORTFM_API int RTFM_fstat(int Fd, struct stat* StatBuf) throw()
 {
 	FContext* Context = FContext::Get();
 	Context->RecordWrite(StatBuf, sizeof(*StatBuf));
@@ -357,7 +359,7 @@ extern "C" AUTORTFM_INTERNAL int RTFM_fstat(int Fd, struct stat* StatBuf) throw(
 	return fstat(Fd, StatBuf);
 }
 
-extern "C" AUTORTFM_INTERNAL int RTFM___xstat(int Ver, const char* Path, struct stat* StatBuf) throw()
+extern "C" AUTORTFM_INTERNAL UE_AUTORTFM_API int RTFM___xstat(int Ver, const char* Path, struct stat* StatBuf) throw()
 {
 	FContext* Context = FContext::Get();
 	Context->RecordWrite(StatBuf, sizeof(*StatBuf));
@@ -365,7 +367,7 @@ extern "C" AUTORTFM_INTERNAL int RTFM___xstat(int Ver, const char* Path, struct
 	return __xstat(Ver, Path, StatBuf);
 }
 
-extern "C" AUTORTFM_INTERNAL int RTFM___fxstat(int Ver, int Fd, struct stat* StatBuf) throw()
+extern "C" AUTORTFM_INTERNAL UE_AUTORTFM_API int RTFM___fxstat(int Ver, int Fd, struct stat* StatBuf) throw()
 {
 	FContext* Context = FContext::Get();
 	Context->RecordWrite(StatBuf, sizeof(*StatBuf));
@@ -373,13 +375,13 @@ extern "C" AUTORTFM_INTERNAL int RTFM___fxstat(int Ver, int Fd, struct stat* Sta
 	return __fxstat(Ver, Fd, StatBuf);
 }
 
-extern "C" AUTORTFM_INTERNAL float RTFM_strtof32(const char* String, char** EndPtr) throw()
+extern "C" AUTORTFM_INTERNAL UE_AUTORTFM_API float RTFM_strtof32(const char* String, char** EndPtr) throw()
 {
 	RecordOpenWriteIfNotNull(EndPtr);
 	return strtof32(String, EndPtr);
 }
 
-extern "C" AUTORTFM_INTERNAL double RTFM_strtof64(const char* String, char** EndPtr) throw()
+extern "C" AUTORTFM_INTERNAL UE_AUTORTFM_API double RTFM_strtof64(const char* String, char** EndPtr) throw()
 {
 	RecordOpenWriteIfNotNull(EndPtr);
 	return strtof64(String, EndPtr);
@@ -388,22 +390,22 @@ extern "C" AUTORTFM_INTERNAL double RTFM_strtof64(const char* String, char** End
 // ^^^^^^^^^^^^^^^^^^^^^^^^^^^^^^ AUTORTFM_PLATFORM_LINUX ^^^^^^^^^^^^^^^^^^^^^^^^^^^^^^
 #endif
 
-extern "C" AUTORTFM_INTERNAL void* RTFM_memcpy(void* Dst, const void* Src, size_t Size) throw()
+extern "C" AUTORTFM_INTERNAL UE_AUTORTFM_API void* RTFM_memcpy(void* Dst, const void* Src, size_t Size) throw()
 {
     return Memcpy(Dst, Src, Size, FContext::Get());
 }
 
-extern "C" AUTORTFM_INTERNAL void* RTFM_memmove(void* Dst, const void* Src, size_t Size) throw()
+extern "C" AUTORTFM_INTERNAL UE_AUTORTFM_API void* RTFM_memmove(void* Dst, const void* Src, size_t Size) throw()
 {
     return Memmove(Dst, Src, Size, FContext::Get());
 }
 
-extern "C" AUTORTFM_INTERNAL void* RTFM_memset(void* Dst, int Value, size_t Size) throw()
+extern "C" AUTORTFM_INTERNAL UE_AUTORTFM_API void* RTFM_memset(void* Dst, int Value, size_t Size) throw()
 {
     return Memset(Dst, Value, Size, FContext::Get());
 }
 
-extern "C" AUTORTFM_INTERNAL void* RTFM_malloc(size_t Size) throw()
+extern "C" AUTORTFM_INTERNAL UE_AUTORTFM_API void* RTFM_malloc(size_t Size) throw()
 {
     void* Result = malloc(Size);
 	FContext* Context = FContext::Get();
@@ -415,7 +417,7 @@ extern "C" AUTORTFM_INTERNAL void* RTFM_malloc(size_t Size) throw()
     return Result;
 }
 
-extern "C" AUTORTFM_INTERNAL void* RTFM_calloc(size_t Count, size_t Size) throw()
+extern "C" AUTORTFM_INTERNAL UE_AUTORTFM_API void* RTFM_calloc(size_t Count, size_t Size) throw()
 {
     void* Result = calloc(Count, Size);
 	FContext* Context = FContext::Get();
@@ -427,7 +429,7 @@ extern "C" AUTORTFM_INTERNAL void* RTFM_calloc(size_t Count, size_t Size) throw(
     return Result;
 }
 
-extern "C" AUTORTFM_INTERNAL void RTFM_free(void* Ptr) throw()
+extern "C" AUTORTFM_INTERNAL UE_AUTORTFM_API void RTFM_free(void* Ptr) throw()
 {
     if (Ptr)
     {
@@ -439,7 +441,7 @@ extern "C" AUTORTFM_INTERNAL void RTFM_free(void* Ptr) throw()
     }
 }
 
-extern "C" AUTORTFM_INTERNAL void* RTFM_realloc(void* Ptr, size_t Size) throw()
+extern "C" AUTORTFM_INTERNAL UE_AUTORTFM_API void* RTFM_realloc(void* Ptr, size_t Size) throw()
 {
     void* NewObject = RTFM_malloc(Size);
     if (Ptr)
@@ -458,7 +460,7 @@ extern "C" AUTORTFM_INTERNAL void* RTFM_realloc(void* Ptr, size_t Size) throw()
     return NewObject;
 }
 
-extern "C" AUTORTFM_INTERNAL char* RTFM_strcpy(char* const Dst, const char* const Src) throw()
+extern "C" AUTORTFM_INTERNAL UE_AUTORTFM_API char* RTFM_strcpy(char* const Dst, const char* const Src) throw()
 {
     const size_t SrcLen = strlen(Src);
 
@@ -467,14 +469,14 @@ extern "C" AUTORTFM_INTERNAL char* RTFM_strcpy(char* const Dst, const char* cons
     return strcpy(Dst, Src);
 }
 
-extern "C" AUTORTFM_INTERNAL char* RTFM_strncpy(char* const Dst, const char* const Src, const size_t Num) throw()
+extern "C" AUTORTFM_INTERNAL UE_AUTORTFM_API char* RTFM_strncpy(char* const Dst, const char* const Src, const size_t Num) throw()
 {
 	FContext* Context = FContext::Get();
     Context->RecordWrite(Dst, Num);
     return strncpy(Dst, Src, Num);
 }
 
-extern "C" AUTORTFM_INTERNAL char* RTFM_strcat(char* const Dst, const char* const Src) throw()
+extern "C" AUTORTFM_INTERNAL UE_AUTORTFM_API char* RTFM_strcat(char* const Dst, const char* const Src) throw()
 {
     const size_t DstLen = strlen(Dst);
     const size_t SrcLen = strlen(Src);
@@ -484,7 +486,7 @@ extern "C" AUTORTFM_INTERNAL char* RTFM_strcat(char* const Dst, const char* cons
     return strcat(Dst, Src);
 }
 
-extern "C" AUTORTFM_INTERNAL char* RTFM_strncat(char* const Dst, const char* const Src, const size_t Num) throw()
+extern "C" AUTORTFM_INTERNAL UE_AUTORTFM_API char* RTFM_strncat(char* const Dst, const char* const Src, const size_t Num) throw()
 {
     const size_t DstLen = strlen(Dst);
 
@@ -493,200 +495,200 @@ extern "C" AUTORTFM_INTERNAL char* RTFM_strncat(char* const Dst, const char* con
     return strncat(Dst, Src, Num);
 }
 
-extern "C" AUTORTFM_INTERNAL long RTFM_strtol(const char* String, char** EndPtr, int Radix) throw()
+extern "C" AUTORTFM_INTERNAL UE_AUTORTFM_API long RTFM_strtol(const char* String, char** EndPtr, int Radix) throw()
 { 
     RecordOpenWriteIfNotNull(EndPtr);
 	return strtol(String, EndPtr, Radix); 
 }
 
-extern "C" AUTORTFM_INTERNAL long long RTFM_strtoll(const char* String, char** EndPtr, int Radix) throw()
+extern "C" AUTORTFM_INTERNAL UE_AUTORTFM_API long long RTFM_strtoll(const char* String, char** EndPtr, int Radix) throw()
 { 
     RecordOpenWriteIfNotNull(EndPtr);
 	return strtoll(String, EndPtr, Radix); 
 }
 
-extern "C" AUTORTFM_INTERNAL unsigned long RTFM_strtoul(const char* String, char** EndPtr, int Radix) throw()
+extern "C" AUTORTFM_INTERNAL UE_AUTORTFM_API unsigned long RTFM_strtoul(const char* String, char** EndPtr, int Radix) throw()
 { 
     RecordOpenWriteIfNotNull(EndPtr);
 	return strtoul(String, EndPtr, Radix); 
 }
 
-extern "C" AUTORTFM_INTERNAL unsigned long long RTFM_strtoull(const char* String, char** EndPtr, int Radix) throw()
+extern "C" AUTORTFM_INTERNAL UE_AUTORTFM_API unsigned long long RTFM_strtoull(const char* String, char** EndPtr, int Radix) throw()
 { 
     RecordOpenWriteIfNotNull(EndPtr);
 	return strtoull(String, EndPtr, Radix); 
 }
 
-extern "C" AUTORTFM_INTERNAL float RTFM_strtof(const char* String, char** EndPtr) throw()
+extern "C" AUTORTFM_INTERNAL UE_AUTORTFM_API float RTFM_strtof(const char* String, char** EndPtr) throw()
 { 
     RecordOpenWriteIfNotNull(EndPtr);
 	return strtof(String, EndPtr); 
 }
 
-extern "C" AUTORTFM_INTERNAL double RTFM_strtod(const char* String, char** EndPtr) throw()
+extern "C" AUTORTFM_INTERNAL UE_AUTORTFM_API double RTFM_strtod(const char* String, char** EndPtr) throw()
 { 
     RecordOpenWriteIfNotNull(EndPtr);
 	return strtod(String, EndPtr); 
 }
 
-extern "C" AUTORTFM_INTERNAL double RTFM_wcstod(const wchar_t* String, wchar_t** EndPtr) throw()
+extern "C" AUTORTFM_INTERNAL UE_AUTORTFM_API double RTFM_wcstod(const wchar_t* String, wchar_t** EndPtr) throw()
 { 
     RecordOpenWriteIfNotNull(EndPtr);
 	return wcstod(String, EndPtr); 
 }
 
-extern "C" AUTORTFM_INTERNAL float RTFM_wcstof(const wchar_t* String, wchar_t** EndPtr) throw()
+extern "C" AUTORTFM_INTERNAL UE_AUTORTFM_API float RTFM_wcstof(const wchar_t* String, wchar_t** EndPtr) throw()
 { 
     RecordOpenWriteIfNotNull(EndPtr);
 	return wcstof(String, EndPtr); 
 }
 
-extern "C" AUTORTFM_INTERNAL long double RTFM_wcstold(const wchar_t* String, wchar_t** EndPtr) throw()
+extern "C" AUTORTFM_INTERNAL UE_AUTORTFM_API long double RTFM_wcstold(const wchar_t* String, wchar_t** EndPtr) throw()
 { 
     RecordOpenWriteIfNotNull(EndPtr);
 	return wcstold(String, EndPtr); 
 }
 
-extern "C" AUTORTFM_INTERNAL long RTFM_wcstol(const wchar_t* String, wchar_t** EndPtr, int Radix) throw()
+extern "C" AUTORTFM_INTERNAL UE_AUTORTFM_API long RTFM_wcstol(const wchar_t* String, wchar_t** EndPtr, int Radix) throw()
 { 
     RecordOpenWriteIfNotNull(EndPtr);
 	return wcstol(String, EndPtr, Radix); 
 }
 
-extern "C" AUTORTFM_INTERNAL long long RTFM_wcstoll(const wchar_t* String, wchar_t** EndPtr, int Radix) throw()
+extern "C" AUTORTFM_INTERNAL UE_AUTORTFM_API long long RTFM_wcstoll(const wchar_t* String, wchar_t** EndPtr, int Radix) throw()
 { 
     RecordOpenWriteIfNotNull(EndPtr);
 	return wcstoll(String, EndPtr, Radix); 
 }
 
-extern "C" AUTORTFM_INTERNAL unsigned long RTFM_wcstoul(const wchar_t* String, wchar_t** EndPtr, int Radix) throw()
+extern "C" AUTORTFM_INTERNAL UE_AUTORTFM_API unsigned long RTFM_wcstoul(const wchar_t* String, wchar_t** EndPtr, int Radix) throw()
 { 
     RecordOpenWriteIfNotNull(EndPtr);
 	return wcstoul(String, EndPtr, Radix); 
 }
 
-extern "C" AUTORTFM_INTERNAL unsigned long long RTFM_wcstoull(const wchar_t* String, wchar_t** EndPtr, int Radix) throw()
+extern "C" AUTORTFM_INTERNAL UE_AUTORTFM_API unsigned long long RTFM_wcstoull(const wchar_t* String, wchar_t** EndPtr, int Radix) throw()
 { 
     RecordOpenWriteIfNotNull(EndPtr);
 	return wcstoull(String, EndPtr, Radix); 
 }
 
-extern "C" AUTORTFM_INTERNAL std::to_chars_result RTFM__ToChars_i(char* First, char* Last, int Value, int Base)
+extern "C" AUTORTFM_INTERNAL UE_AUTORTFM_API std::to_chars_result RTFM__ToChars_i(char* First, char* Last, int Value, int Base)
 {
 	RecordOpenWrite(First, Last - First);
 	return std::to_chars(First, Last, Value, Base);
 }
 
-extern "C" AUTORTFM_INTERNAL std::to_chars_result RTFM__ToChars_u(char* First, char* Last, unsigned int Value, int Base)
+extern "C" AUTORTFM_INTERNAL UE_AUTORTFM_API std::to_chars_result RTFM__ToChars_u(char* First, char* Last, unsigned int Value, int Base)
 {
 	RecordOpenWrite(First, Last - First);
 	return std::to_chars(First, Last, Value, Base);
 }
 
-extern "C" AUTORTFM_INTERNAL std::to_chars_result RTFM__ToChars_i8(char* First, char* Last, int8_t Value, int Base)
+extern "C" AUTORTFM_INTERNAL UE_AUTORTFM_API std::to_chars_result RTFM__ToChars_i8(char* First, char* Last, int8_t Value, int Base)
 {
 	RecordOpenWrite(First, Last - First);
 	return std::to_chars(First, Last, Value, Base);
 }
 
-extern "C" AUTORTFM_INTERNAL std::to_chars_result RTFM__ToChars_u8(char* First, char* Last, uint8_t Value, int Base)
+extern "C" AUTORTFM_INTERNAL UE_AUTORTFM_API std::to_chars_result RTFM__ToChars_u8(char* First, char* Last, uint8_t Value, int Base)
 {
 	RecordOpenWrite(First, Last - First);
 	return std::to_chars(First, Last, Value, Base);
 }
 
-extern "C" AUTORTFM_INTERNAL std::to_chars_result RTFM__ToChars_i16(char* First, char* Last, int16_t Value, int Base)
+extern "C" AUTORTFM_INTERNAL UE_AUTORTFM_API std::to_chars_result RTFM__ToChars_i16(char* First, char* Last, int16_t Value, int Base)
 {
 	RecordOpenWrite(First, Last - First);
 	return std::to_chars(First, Last, Value, Base);
 }
 
-extern "C" AUTORTFM_INTERNAL std::to_chars_result RTFM__ToChars_u16(char* First, char* Last, uint16_t Value, int Base)
+extern "C" AUTORTFM_INTERNAL UE_AUTORTFM_API std::to_chars_result RTFM__ToChars_u16(char* First, char* Last, uint16_t Value, int Base)
 {
 	RecordOpenWrite(First, Last - First);
 	return std::to_chars(First, Last, Value, Base);
 }
 
-extern "C" AUTORTFM_INTERNAL std::to_chars_result RTFM__ToChars_i32(char* First, char* Last, int32_t Value, int Base)
+extern "C" AUTORTFM_INTERNAL UE_AUTORTFM_API std::to_chars_result RTFM__ToChars_i32(char* First, char* Last, int32_t Value, int Base)
 {
 	RecordOpenWrite(First, Last - First);
 	return std::to_chars(First, Last, Value, Base);
 }
 
-extern "C" AUTORTFM_INTERNAL std::to_chars_result RTFM__ToChars_u32(char* First, char* Last, uint32_t Value, int Base)
+extern "C" AUTORTFM_INTERNAL UE_AUTORTFM_API std::to_chars_result RTFM__ToChars_u32(char* First, char* Last, uint32_t Value, int Base)
 {
 	RecordOpenWrite(First, Last - First);
 	return std::to_chars(First, Last, Value, Base);
 }
 
-extern "C" AUTORTFM_INTERNAL std::to_chars_result RTFM__ToChars_i64(char* First, char* Last, int64_t Value, int Base)
+extern "C" AUTORTFM_INTERNAL UE_AUTORTFM_API std::to_chars_result RTFM__ToChars_i64(char* First, char* Last, int64_t Value, int Base)
 {
 	RecordOpenWrite(First, Last - First);
 	return std::to_chars(First, Last, Value, Base);
 }
 
-extern "C" AUTORTFM_INTERNAL std::to_chars_result RTFM__ToChars_u64(char* First, char* Last, uint64_t Value, int Base)
+extern "C" AUTORTFM_INTERNAL UE_AUTORTFM_API std::to_chars_result RTFM__ToChars_u64(char* First, char* Last, uint64_t Value, int Base)
 {
 	RecordOpenWrite(First, Last - First);
 	return std::to_chars(First, Last, Value, Base);
 }
 
-extern "C" AUTORTFM_INTERNAL std::to_chars_result RTFM__ToChars_f(char* First, char* Last, float Value)
+extern "C" AUTORTFM_INTERNAL UE_AUTORTFM_API std::to_chars_result RTFM__ToChars_f(char* First, char* Last, float Value)
 {
 	RecordOpenWrite(First, Last - First);
 	return std::to_chars(First, Last, Value);
 }
 
-extern "C" AUTORTFM_INTERNAL std::to_chars_result RTFM__ToChars_f_F(char* First, char* Last, float Value, std::chars_format Format)
+extern "C" AUTORTFM_INTERNAL UE_AUTORTFM_API std::to_chars_result RTFM__ToChars_f_F(char* First, char* Last, float Value, std::chars_format Format)
 {
 	RecordOpenWrite(First, Last - First);
 	return std::to_chars(First, Last, Value, Format);
 }
 
-extern "C" AUTORTFM_INTERNAL std::to_chars_result RTFM__ToChars_f_FP(char* First, char* Last, float Value, std::chars_format Format, int Precision)
+extern "C" AUTORTFM_INTERNAL UE_AUTORTFM_API std::to_chars_result RTFM__ToChars_f_FP(char* First, char* Last, float Value, std::chars_format Format, int Precision)
 {
 	RecordOpenWrite(First, Last - First);
 	return std::to_chars(First, Last, Value, Format, Precision);
 }
 
-extern "C" AUTORTFM_INTERNAL std::to_chars_result RTFM__ToChars_d(char* First, char* Last, double Value)
+extern "C" AUTORTFM_INTERNAL UE_AUTORTFM_API std::to_chars_result RTFM__ToChars_d(char* First, char* Last, double Value)
 {
 	RecordOpenWrite(First, Last - First);
 	return std::to_chars(First, Last, Value);
 }
 
-extern "C" AUTORTFM_INTERNAL std::to_chars_result RTFM__ToChars_d_F(char* First, char* Last, double Value, std::chars_format Format)
+extern "C" AUTORTFM_INTERNAL UE_AUTORTFM_API std::to_chars_result RTFM__ToChars_d_F(char* First, char* Last, double Value, std::chars_format Format)
 {
 	RecordOpenWrite(First, Last - First);
 	return std::to_chars(First, Last, Value, Format);
 }
 
-extern "C" AUTORTFM_INTERNAL std::to_chars_result RTFM__ToChars_d_FP(char* First, char* Last, double Value, std::chars_format Format, int Precision)
+extern "C" AUTORTFM_INTERNAL UE_AUTORTFM_API std::to_chars_result RTFM__ToChars_d_FP(char* First, char* Last, double Value, std::chars_format Format, int Precision)
 {
 	RecordOpenWrite(First, Last - First);
 	return std::to_chars(First, Last, Value, Format, Precision);
 }
 
-extern "C" AUTORTFM_INTERNAL std::to_chars_result RTFM__ToChars_ld(char* First, char* Last, long double Value)
+extern "C" AUTORTFM_INTERNAL UE_AUTORTFM_API std::to_chars_result RTFM__ToChars_ld(char* First, char* Last, long double Value)
 {
 	RecordOpenWrite(First, Last - First);
 	return std::to_chars(First, Last, Value);
 }
 
-extern "C" AUTORTFM_INTERNAL std::to_chars_result RTFM__ToChars_ld_F(char* First, char* Last, long double Value, std::chars_format Format)
+extern "C" AUTORTFM_INTERNAL UE_AUTORTFM_API std::to_chars_result RTFM__ToChars_ld_F(char* First, char* Last, long double Value, std::chars_format Format)
 {
 	RecordOpenWrite(First, Last - First);
 	return std::to_chars(First, Last, Value, Format);
 }
 
-extern "C" AUTORTFM_INTERNAL std::to_chars_result RTFM__ToChars_ld_FP(char* First, char* Last, long double Value, std::chars_format Format, int Precision)
+extern "C" AUTORTFM_INTERNAL UE_AUTORTFM_API std::to_chars_result RTFM__ToChars_ld_FP(char* First, char* Last, long double Value, std::chars_format Format, int Precision)
 {
 	RecordOpenWrite(First, Last - First);
 	return std::to_chars(First, Last, Value, Format, Precision);
 }
 
 // FIXME: Does not currently support %n format specifiers.
-extern "C" AUTORTFM_INTERNAL int RTFM_vsnprintf(char* Str, size_t Size, const char* Format, va_list ArgList) throw()
+extern "C" AUTORTFM_INTERNAL UE_AUTORTFM_API int RTFM_vsnprintf(char* Str, size_t Size, const char* Format, va_list ArgList) throw()
 {
     ThrowIfFormatContainsPercentN(Format);
 
@@ -710,7 +712,7 @@ extern "C" AUTORTFM_INTERNAL int RTFM_vsnprintf(char* Str, size_t Size, const ch
 }
 
 // FIXME: Does not currently support %n format specifiers.
-extern "C" AUTORTFM_INTERNAL int RTFM_vswprintf(wchar_t* Str, size_t Size, const wchar_t* Format, va_list ArgList)
+extern "C" AUTORTFM_INTERNAL UE_AUTORTFM_API int RTFM_vswprintf(wchar_t* Str, size_t Size, const wchar_t* Format, va_list ArgList)
 {
     ThrowIfFormatContainsPercentN(Format);
 
@@ -741,7 +743,7 @@ extern "C" AUTORTFM_INTERNAL int RTFM_vswprintf(wchar_t* Str, size_t Size, const
 }
 
 // FIXME: Does not currently support %n format specifiers.
-extern "C" AUTORTFM_INTERNAL int RTFM_swprintf(wchar_t* Buffer, size_t BufferCount, wchar_t const* Format, ...)
+extern "C" AUTORTFM_INTERNAL UE_AUTORTFM_API int RTFM_swprintf(wchar_t* Buffer, size_t BufferCount, wchar_t const* Format, ...)
 {
     va_list ArgList;
 
@@ -753,7 +755,7 @@ extern "C" AUTORTFM_INTERNAL int RTFM_swprintf(wchar_t* Buffer, size_t BufferCou
 }
 
 // FIXME: Does not currently support %n format specifiers.
-extern "C" AUTORTFM_INTERNAL int RTFM_snprintf(char* Str, size_t Size, const char* Format, ...) throw()
+extern "C" AUTORTFM_INTERNAL UE_AUTORTFM_API int RTFM_snprintf(char* Str, size_t Size, const char* Format, ...) throw()
 {
     va_list ArgList;
 
@@ -765,7 +767,7 @@ extern "C" AUTORTFM_INTERNAL int RTFM_snprintf(char* Str, size_t Size, const cha
 }
 
 // FIXME: Does not currently support %n format specifiers.
-extern "C" AUTORTFM_INTERNAL int RTFM_printf(const char* Format, ...)
+extern "C" AUTORTFM_INTERNAL UE_AUTORTFM_API int RTFM_printf(const char* Format, ...)
 {
     ThrowIfFormatContainsPercentN(Format);
 
@@ -777,7 +779,7 @@ extern "C" AUTORTFM_INTERNAL int RTFM_printf(const char* Format, ...)
 }
 
 // FIXME: Does not currently support %n format specifiers.
-extern "C" AUTORTFM_INTERNAL int RTFM_wprintf(const wchar_t* Format, ...)
+extern "C" AUTORTFM_INTERNAL UE_AUTORTFM_API int RTFM_wprintf(const wchar_t* Format, ...)
 {
     ThrowIfFormatContainsPercentN(Format);
 
@@ -788,7 +790,7 @@ extern "C" AUTORTFM_INTERNAL int RTFM_wprintf(const wchar_t* Format, ...)
     return Result;
 }
 
-extern "C" AUTORTFM_INTERNAL wchar_t* RTFM_wcscpy(wchar_t* Dst, const wchar_t* Src) throw()
+extern "C" AUTORTFM_INTERNAL UE_AUTORTFM_API wchar_t* RTFM_wcscpy(wchar_t* Dst, const wchar_t* Src) throw()
 {
     const size_t SrcLen = wcslen(Src);
 
@@ -797,14 +799,14 @@ extern "C" AUTORTFM_INTERNAL wchar_t* RTFM_wcscpy(wchar_t* Dst, const wchar_t* S
     return wcscpy(Dst, Src);
 }
 
-extern "C" AUTORTFM_INTERNAL wchar_t* RTFM_wcsncpy(wchar_t* Dst, const wchar_t* Src, size_t Count) throw()
+extern "C" AUTORTFM_INTERNAL UE_AUTORTFM_API wchar_t* RTFM_wcsncpy(wchar_t* Dst, const wchar_t* Src, size_t Count) throw()
 {
 	FContext* Context = FContext::Get();
 	Context->RecordWrite(Dst, Count * sizeof(wchar_t));
 	return wcsncpy(Dst, Src, Count);
 }
 
-extern "C" AUTORTFM_INTERNAL int RTFM_atexit(void(*Callback)(void)) throw()
+extern "C" AUTORTFM_INTERNAL UE_AUTORTFM_API int RTFM_atexit(void(*Callback)(void)) throw()
 {
 	FContext* Context = FContext::Get();
 	Context->GetCurrentTransaction()->DeferUntilCommit([Callback]
@@ -826,23 +828,6 @@ extern "C" void* __cxa_begin_catch(void* exceptionObject) throw();
 extern "C" void __cxa_end_catch();
 #endif
 
-}  // anonymous namespace
-
-////////////////////////////////////////////////////////////////////////////////
-// Register all the open -> closed functions
-////////////////////////////////////////////////////////////////////////////////
-#if AUTORTFM_PLATFORM_LINUX
-// LibCxx's string.h uses builtin string functions inside inline functions.
-// We can't take the address of these builtins, but we know they map to these C functions.
-extern "C"
-{
-	char* strchr(const char*, int);
-	char* strrchr(const char*, int);
-	char* strstr(const char*, const char*);
-}
-#endif
-
-#if FORCE_ANSI_ALLOCATOR
 template<typename OpenNewFn, typename OpenDeleteFn, typename ... ArgTys>
 UE_AUTORTFM_FORCEINLINE static void* RTFM_ClosedNew(OpenNewFn* New, OpenDeleteFn* Delete, unsigned long Size, ArgTys... Args)
 {
@@ -869,8 +854,16 @@ UE_AUTORTFM_FORCEINLINE static void RTFM_ClosedDelete(OpenDeleteFn* Delete, void
 	}
 }
 
+#if AUTORTFM_PLATFORM_LINUX
+// vvvvvvvvvvvvvvvvvvvvvvvvvvvvvv AUTORTFM_PLATFORM_LINUX vvvvvvvvvvvvvvvvvvvvvvvvvvvvvv
+// LibCxx's string.h uses builtin string functions inside inline functions.
+// We can't take the address of these builtins, but we know they map to these C functions.
 extern "C"
 {
+	char* strchr(const char*, int);
+	char* strrchr(const char*, int);
+	char* strstr(const char*, const char*);
+
 	// Operator new and new[]
 	void* _Znwm(unsigned long);
 	void* _Znam(unsigned long);
@@ -895,108 +888,111 @@ extern "C"
 	void _ZdaPvmSt11align_val_t(void*, unsigned long, std::align_val_t);
 } // extern "C"
 
-extern "C" AUTORTFM_INTERNAL void* RTFM_Znwm(unsigned long Size)
+extern "C" AUTORTFM_INTERNAL UE_AUTORTFM_API void* RTFM_Znwm(unsigned long Size)
 {
 	return RTFM_ClosedNew(_Znwm, _ZdlPv, Size);
 }
 
-extern "C" AUTORTFM_INTERNAL void* RTFM_Znam(unsigned long Size)
+extern "C" AUTORTFM_INTERNAL UE_AUTORTFM_API void* RTFM_Znam(unsigned long Size)
 {
 	return RTFM_ClosedNew(_Znam, _ZdaPv, Size);
 }
 
-extern "C" AUTORTFM_INTERNAL void* RTFM_ZnwmRKSt9nothrow_t(unsigned long Size, const std::nothrow_t& NoThrow)
+extern "C" AUTORTFM_INTERNAL UE_AUTORTFM_API void* RTFM_ZnwmRKSt9nothrow_t(unsigned long Size, const std::nothrow_t& NoThrow)
 {
 	return RTFM_ClosedNew(_ZnwmRKSt9nothrow_t, _ZdlPvRKSt9nothrow_t, Size, NoThrow);
 }
 
-extern "C" AUTORTFM_INTERNAL void* RTFM_ZnamRKSt9nothrow_t(unsigned long Size, const std::nothrow_t& NoThrow)
+extern "C" AUTORTFM_INTERNAL UE_AUTORTFM_API void* RTFM_ZnamRKSt9nothrow_t(unsigned long Size, const std::nothrow_t& NoThrow)
 {
 	return RTFM_ClosedNew(_ZnamRKSt9nothrow_t, _ZdaPvRKSt9nothrow_t, Size, NoThrow);
 }
 
-extern "C" AUTORTFM_INTERNAL void* RTFM_ZnwmSt11align_val_t(unsigned long Size, std::align_val_t Align)
+extern "C" AUTORTFM_INTERNAL UE_AUTORTFM_API void* RTFM_ZnwmSt11align_val_t(unsigned long Size, std::align_val_t Align)
 {
 	return RTFM_ClosedNew(_ZnwmSt11align_val_t, _ZdlPvSt11align_val_t, Size, Align);
 }
 
-extern "C" AUTORTFM_INTERNAL void* RTFM_ZnamSt11align_val_t(unsigned long Size, std::align_val_t Align)
+extern "C" AUTORTFM_INTERNAL UE_AUTORTFM_API void* RTFM_ZnamSt11align_val_t(unsigned long Size, std::align_val_t Align)
 {
 	return RTFM_ClosedNew(_ZnamSt11align_val_t, _ZdaPvSt11align_val_t, Size, Align);
 }
 
-extern "C" AUTORTFM_INTERNAL void* RTFM_ZnwmSt11align_val_tRKSt9nothrow_t(unsigned long Size, std::align_val_t Align, const std::nothrow_t& NoThrow)
+extern "C" AUTORTFM_INTERNAL UE_AUTORTFM_API void* RTFM_ZnwmSt11align_val_tRKSt9nothrow_t(unsigned long Size, std::align_val_t Align, const std::nothrow_t& NoThrow)
 {
 	return RTFM_ClosedNew(_ZnwmSt11align_val_tRKSt9nothrow_t, _ZdlPvSt11align_val_tRKSt9nothrow_t, Size, Align, NoThrow);
 }
 
-extern "C" AUTORTFM_INTERNAL void* RTFM_ZnamSt11align_val_tRKSt9nothrow_t(unsigned long Size, std::align_val_t Align, const std::nothrow_t& NoThrow)
+extern "C" AUTORTFM_INTERNAL UE_AUTORTFM_API void* RTFM_ZnamSt11align_val_tRKSt9nothrow_t(unsigned long Size, std::align_val_t Align, const std::nothrow_t& NoThrow)
 {
 	return RTFM_ClosedNew(_ZnamSt11align_val_tRKSt9nothrow_t, _ZdaPvSt11align_val_tRKSt9nothrow_t, Size, Align, NoThrow);
 }
 
-extern "C" AUTORTFM_INTERNAL void RTFM_ZdlPv(void* Pointer)
+extern "C" AUTORTFM_INTERNAL UE_AUTORTFM_API void RTFM_ZdlPv(void* Pointer)
 {
 	RTFM_ClosedDelete(_ZdlPv, Pointer);
 }
 
-extern "C" AUTORTFM_INTERNAL void RTFM_ZdaPv(void* Pointer)
+extern "C" AUTORTFM_INTERNAL UE_AUTORTFM_API void RTFM_ZdaPv(void* Pointer)
 {
 	RTFM_ClosedDelete(_ZdaPv, Pointer);
 }
 
-extern "C" AUTORTFM_INTERNAL void RTFM_ZdlPvRKSt9nothrow_t(void* Pointer, const std::nothrow_t& NoThrow)
+extern "C" AUTORTFM_INTERNAL UE_AUTORTFM_API void RTFM_ZdlPvRKSt9nothrow_t(void* Pointer, const std::nothrow_t& NoThrow)
 {
 	RTFM_ClosedDelete(_ZdlPvRKSt9nothrow_t, Pointer, NoThrow);
 }
 
-extern "C" AUTORTFM_INTERNAL void RTFM_ZdaPvRKSt9nothrow_t(void* Pointer, const std::nothrow_t& NoThrow)
+extern "C" AUTORTFM_INTERNAL UE_AUTORTFM_API void RTFM_ZdaPvRKSt9nothrow_t(void* Pointer, const std::nothrow_t& NoThrow)
 {
 	RTFM_ClosedDelete(_ZdaPvRKSt9nothrow_t, Pointer, NoThrow);
 }
 
-extern "C" AUTORTFM_INTERNAL void RTFM_ZdlPvm(void* Pointer, unsigned long Size)
+extern "C" AUTORTFM_INTERNAL UE_AUTORTFM_API void RTFM_ZdlPvm(void* Pointer, unsigned long Size)
 {
 	RTFM_ClosedDelete(_ZdlPvm, Pointer, Size);
 }
 
-extern "C" AUTORTFM_INTERNAL void RTFM_ZdaPvm(void* Pointer, unsigned long Size)
+extern "C" AUTORTFM_INTERNAL UE_AUTORTFM_API void RTFM_ZdaPvm(void* Pointer, unsigned long Size)
 {
 	RTFM_ClosedDelete(_ZdaPvm, Pointer, Size);
 }
 
-extern "C" AUTORTFM_INTERNAL void RTFM_ZdlPvSt11align_val_t(void* Pointer, std::align_val_t Align)
+extern "C" AUTORTFM_INTERNAL UE_AUTORTFM_API void RTFM_ZdlPvSt11align_val_t(void* Pointer, std::align_val_t Align)
 {
 	RTFM_ClosedDelete(_ZdlPvSt11align_val_t, Pointer, Align);
 }
 
-extern "C" AUTORTFM_INTERNAL void RTFM_ZdaPvSt11align_val_t(void* Pointer, std::align_val_t Align)
+extern "C" AUTORTFM_INTERNAL UE_AUTORTFM_API void RTFM_ZdaPvSt11align_val_t(void* Pointer, std::align_val_t Align)
 {
 	RTFM_ClosedDelete(_ZdaPvSt11align_val_t, Pointer, Align);
 }
 
-extern "C" AUTORTFM_INTERNAL void RTFM_ZdlPvSt11align_val_tRKSt9nothrow_t(void* Pointer, std::align_val_t Align, const std::nothrow_t& NoThrow)
+extern "C" AUTORTFM_INTERNAL UE_AUTORTFM_API void RTFM_ZdlPvSt11align_val_tRKSt9nothrow_t(void* Pointer, std::align_val_t Align, const std::nothrow_t& NoThrow)
 {
 	RTFM_ClosedDelete(_ZdlPvSt11align_val_tRKSt9nothrow_t, Pointer, Align, NoThrow);
 }
 
-extern "C" AUTORTFM_INTERNAL void RTFM_ZdaPvSt11align_val_tRKSt9nothrow_t(void* Pointer, std::align_val_t Align, const std::nothrow_t& NoThrow)
+extern "C" AUTORTFM_INTERNAL UE_AUTORTFM_API void RTFM_ZdaPvSt11align_val_tRKSt9nothrow_t(void* Pointer, std::align_val_t Align, const std::nothrow_t& NoThrow)
 {
 	RTFM_ClosedDelete(_ZdaPvSt11align_val_tRKSt9nothrow_t, Pointer, Align, NoThrow);
 }
 
-extern "C" AUTORTFM_INTERNAL void RTFM_ZdlPvmSt11align_val_t(void* Pointer, unsigned long Size, std::align_val_t Align)
+extern "C" AUTORTFM_INTERNAL UE_AUTORTFM_API void RTFM_ZdlPvmSt11align_val_t(void* Pointer, unsigned long Size, std::align_val_t Align)
 {
 	RTFM_ClosedDelete(_ZdlPvmSt11align_val_t, Pointer, Size, Align);
 }
 
-extern "C" AUTORTFM_INTERNAL void RTFM_ZdaPvmSt11align_val_t(void* Pointer, unsigned long Size, std::align_val_t Align)
+extern "C" AUTORTFM_INTERNAL UE_AUTORTFM_API void RTFM_ZdaPvmSt11align_val_t(void* Pointer, unsigned long Size, std::align_val_t Align)
 {
 	RTFM_ClosedDelete(_ZdaPvmSt11align_val_t, Pointer, Size, Align);
 }
+// ^^^^^^^^^^^^^^^^^^^^^^^^^^^^^^ AUTORTFM_PLATFORM_LINUX ^^^^^^^^^^^^^^^^^^^^^^^^^^^^^^
+#endif
 
-#endif // FORCE_ANSI_ALLOCATOR
-
+////////////////////////////////////////////////////////////////////////////////
+// Register all the open -> closed functions
+////////////////////////////////////////////////////////////////////////////////
 UE_AUTORTFM_REGISTER_OPEN_TO_CLOSED_FUNCTIONS(
 #if AUTORTFM_ASAN_ENABLED
 	UE_AUTORTFM_MAP_OPEN_TO_SELF(__asan_addr_is_in_fake_stack),
@@ -1099,7 +1095,7 @@ UE_AUTORTFM_REGISTER_OPEN_TO_CLOSED_FUNCTIONS(
 	UE_AUTORTFM_MAP_OPEN_TO_SELF_OVERLOADED(const char* (const char*, const char*), strstr),
 #endif
 
-#if FORCE_ANSI_ALLOCATOR
+#if AUTORTFM_PLATFORM_LINUX
 	UE_AUTORTFM_MAP_OPEN_TO_CLOSED(_Znwm, RTFM_Znwm),
 	UE_AUTORTFM_MAP_OPEN_TO_CLOSED(_Znam, RTFM_Znam),
 	UE_AUTORTFM_MAP_OPEN_TO_CLOSED(_ZnwmRKSt9nothrow_t, RTFM_ZnwmRKSt9nothrow_t),
