[AutoRTFM] Remove PostAbortCallback.

AutoRTFM::CascadingRetryTransaction and AutoRTFM::CascadingAbortTransaction no longer take a callback directly. Instead, you can call AutoRTFM::OnComplete immediately before invoking them.

The OnComplete mechanism turns out to be functionally identical but more general-purpose.

diff --git a/Engine/Source/Programs/AutoRTFMTests/Private/API.cpp b/Engine/Source/Programs/AutoRTFMTests/Private/API.cpp
--- a/Engine/Source/Programs/AutoRTFMTests/Private/API.cpp
+++ b/Engine/Source/Programs/AutoRTFMTests/Private/API.cpp
@@ -1150,13 +1150,16 @@ TEST_CASE("API.IsOnCurrentTransactionStack")
 	});
 }
 
-TEST_CASE("API.CascadingRetryTransaction")
+TEST_CASE("API.OnCompleteIgnoredOutsideTransaction")
 {
 	SECTION("Callback not called outside transaction")
 	{
-		AutoRTFM::CascadingRetryTransaction([&] { FAIL("Unreachable"); });
+		AutoRTFM::OnComplete([&] { FAIL("Unreachable"); });
 	}
+}
 
+TEST_CASE("API.CascadingRetryTransaction")
+{
 	SECTION("Non-nested committed transaction")
 	{
 		bool bFirst = true;
@@ -1165,12 +1168,13 @@ TEST_CASE("API.CascadingRetryTransaction")
 			if (bFirst)
 			{
 				bFirst = false;
-				AutoRTFM::CascadingRetryTransaction([&]
+				AutoRTFM::OnComplete([&]
 				{
 					// The write above would have been undone.
 					REQUIRE(bFirst);
 					bFirst = false;
 				});
+				AutoRTFM::CascadingRetryTransaction();
 			}
 		});
 		REQUIRE(!bFirst);
@@ -1184,12 +1188,13 @@ TEST_CASE("API.CascadingRetryTransaction")
 			if (bFirst)
 			{
 				bFirst = false;
-				AutoRTFM::CascadingRetryTransaction([&]
+				AutoRTFM::OnComplete([&]
 				{
 					// The write above would have been undone.
 					REQUIRE(bFirst);
 					bFirst = false;
 				});
+				AutoRTFM::CascadingRetryTransaction();
 			}
 
 			AutoRTFM::AbortTransaction();
@@ -1207,12 +1212,13 @@ TEST_CASE("API.CascadingRetryTransaction")
 				if (bFirst)
 				{
 					bFirst = false;
-					AutoRTFM::CascadingRetryTransaction([&]
+					AutoRTFM::OnComplete([&]
 					{
 						// The write above would have been undone.
 						REQUIRE(bFirst);
 						bFirst = false;
 					});
+					AutoRTFM::CascadingRetryTransaction();
 				}
 			});
 		});
@@ -1229,12 +1235,13 @@ TEST_CASE("API.CascadingRetryTransaction")
 				if (bFirst)
 				{
 					bFirst = false;
-					AutoRTFM::CascadingRetryTransaction([&]
+					AutoRTFM::OnComplete([&]
 					{
 						// The write above would have been undone.
 						REQUIRE(bFirst);
 						bFirst = false;
 					});
+					AutoRTFM::CascadingRetryTransaction();
 				}
 
 				AutoRTFM::AbortTransaction();
@@ -1250,11 +1257,12 @@ TEST_CASE("API.CascadingRetryTransaction")
 		{
 			if (bFirst)
 			{
-				AutoRTFM::CascadingRetryTransaction([&]
+				AutoRTFM::OnComplete([&]
 				{
 					REQUIRE(!AutoRTFM::IsTransactional());
 					bFirst = false;
 				});
+				AutoRTFM::CascadingRetryTransaction();
 			}
 		});
 		REQUIRE(!bFirst);
@@ -1267,11 +1275,12 @@ TEST_CASE("API.CascadingRetryTransaction")
 		{
 			if (bFirst)
 			{
-				AutoRTFM::CascadingRetryTransaction([&]
+				AutoRTFM::OnComplete([&]
 				{
 					REQUIRE(!AutoRTFM::IsClosed());
 					bFirst = false;
 				});
+				AutoRTFM::CascadingRetryTransaction();
 			}
 		});
 		REQUIRE(!bFirst);
diff --git a/Engine/Source/Programs/AutoRTFMTests/Private/Abort.cpp b/Engine/Source/Programs/AutoRTFMTests/Private/Abort.cpp
--- a/Engine/Source/Programs/AutoRTFMTests/Private/Abort.cpp
+++ b/Engine/Source/Programs/AutoRTFMTests/Private/Abort.cpp
@@ -244,7 +244,7 @@ TEST_CASE("Abort.Cascade")
 	REQUIRE(false == bTouched);
 }
 
-TEST_CASE("Abort.CascadeWithPostAbortCallback")
+TEST_CASE("Abort.CascadeWithOnCompleteCallback")
 {
 	bool bTouched = false;
 	bool bPostAbortCallbackWasExecuted = false;
@@ -254,11 +254,12 @@ TEST_CASE("Abort.CascadeWithPostAbortCallback")
 		bTouched = true;
 		AutoRTFM::Transact([&]
 		{
-			AutoRTFM::CascadingAbortTransaction([&]
+			AutoRTFM::OnComplete([&]
 			{
 				REQUIRE(!bTouched);  // Rollbacks must occur before the post-abort callback is invoked.
 				bPostAbortCallbackWasExecuted = true;
 			});
+			AutoRTFM::CascadingAbortTransaction();
 		});
 
 		FAIL("Execution should never reach this point");
@@ -269,7 +270,7 @@ TEST_CASE("Abort.CascadeWithPostAbortCallback")
 	REQUIRE(bPostAbortCallbackWasExecuted);
 }
 
-TEST_CASE("Abort.CascadingAbortMustRunOnAbortsBeforePostAbort")
+TEST_CASE("Abort.CascadingAbortMustRunOnAbortsBeforeOnComplete")
 {
 	bool bOnAbortWasCalled = false;
 	bool bPostAbortCallbackWasExecuted = false;
@@ -280,12 +281,12 @@ TEST_CASE("Abort.CascadingAbortMustRunOnAbortsBeforePostAbort")
 		{
 			bOnAbortWasCalled = true;
 		});
-
-		AutoRTFM::CascadingAbortTransaction([&]
+		AutoRTFM::OnComplete([&]
 		{
 			REQUIRE(bOnAbortWasCalled);  // OnAborts must before the post-abort callback is invoked.
 			bPostAbortCallbackWasExecuted = true;
 		});
+		AutoRTFM::CascadingAbortTransaction();
 	});
 
 	REQUIRE(AutoRTFM::ETransactionResult::AbortedByCascade == Result);
@@ -342,10 +343,10 @@ TEST_CASE("Abort.CascadeThroughOpen")
 	REQUIRE(false == bTouched);
 }
 
-TEST_CASE("Abort.CascadeWithCallbackThroughOpen")
+TEST_CASE("Abort.CascadeWithOnCompleteThroughOpen")
 {
 	bool bTouched = false;
-	bool bPostAbortCallbackWasExecuted = false;
+	bool bOnCompleteCallbackWasExecuted = false;
 
 	const AutoRTFM::ETransactionResult Result = AutoRTFM::Transact([&]
 	{
@@ -357,11 +358,12 @@ TEST_CASE("Abort.CascadeWithCallbackThroughOpen")
 			{
 				AutoRTFM::Transact([&]
 				{
-					AutoRTFM::CascadingAbortTransaction([&]
+					AutoRTFM::OnComplete([&]
 					{
 						REQUIRE(!bTouched);  // Rollbacks must occur before the post-abort callback is invoked.
-						bPostAbortCallbackWasExecuted = true;
+						bOnCompleteCallbackWasExecuted = true;
 					});
+					AutoRTFM::CascadingAbortTransaction();
 				});
 			});
 
@@ -371,12 +373,12 @@ TEST_CASE("Abort.CascadeWithCallbackThroughOpen")
 
 	REQUIRE(AutoRTFM::ETransactionResult::AbortedByCascade == Result);
 	REQUIRE(!bTouched);
-	REQUIRE(bPostAbortCallbackWasExecuted);
+	REQUIRE(bOnCompleteCallbackWasExecuted);
 }
 
 TEST_CASE("Abort.CascadeThroughManualTransaction")
 {
-	bool bPostAbortCallbackWasExecuted = false;
+	bool bOnCompleteCallbackWasExecuted = false;
 	bool bTouched = false;
 
 	const AutoRTFM::ETransactionResult Result = AutoRTFM::Transact([&]
@@ -389,11 +391,12 @@ TEST_CASE("Abort.CascadeThroughManualTransaction")
 
 			const AutoRTFM::EContextStatus Status = AutoRTFM::Close([&]
 			{
-				AutoRTFM::CascadingAbortTransaction([&]
+				AutoRTFM::OnComplete([&]
 				{
 					REQUIRE(!bTouched);  // Rollbacks must occur before the post-abort callback is invoked.
-					bPostAbortCallbackWasExecuted = true;
+					bOnCompleteCallbackWasExecuted = true;
 				});
+				AutoRTFM::CascadingAbortTransaction();
 			});
 
 			REQUIRE(AutoRTFM::EContextStatus::AbortedByCascadingAbort == Status);
@@ -408,7 +411,7 @@ TEST_CASE("Abort.CascadeThroughManualTransaction")
 
 	REQUIRE(AutoRTFM::ETransactionResult::AbortedByCascade == Result);
 	REQUIRE(!bTouched);
-	REQUIRE(bPostAbortCallbackWasExecuted);
+	REQUIRE(bOnCompleteCallbackWasExecuted);
 }
 
 static int SAbortKey;
diff --git a/Engine/Source/Programs/AutoRTFMTests/Private/UObjectTests.cpp b/Engine/Source/Programs/AutoRTFMTests/Private/UObjectTests.cpp
--- a/Engine/Source/Programs/AutoRTFMTests/Private/UObjectTests.cpp
+++ b/Engine/Source/Programs/AutoRTFMTests/Private/UObjectTests.cpp
@@ -859,7 +859,11 @@ TEST_CASE("UObject.ReleaseRef")
 
 			if (bFirst)
 			{
-				AutoRTFM::CascadingRetryTransaction([&] { bFirst = false; });
+				AutoRTFM::OnComplete([&] 
+				{ 
+					bFirst = false; 
+				});
+				AutoRTFM::CascadingRetryTransaction();
 			}
 		});
 	}
diff --git a/Engine/Source/Runtime/AutoRTFM/Private/API.cpp b/Engine/Source/Runtime/AutoRTFM/Private/API.cpp
--- a/Engine/Source/Runtime/AutoRTFM/Private/API.cpp
+++ b/Engine/Source/Runtime/AutoRTFM/Private/API.cpp
@@ -598,18 +598,18 @@ AUTORTFM_INTERNAL void RTFM_autortfm_record_open_write_err(void*, size_t)
 	AUTORTFM_FATAL("The function `autortfm_record_open_write` was called from closed code");
 }
 
-AUTORTFM_INTERNAL void RTFM_CascadingAbortTransactionInternal(TTask<void()>&& RunAfterAbort)
+AUTORTFM_INTERNAL void RTFM_CascadingAbortTransactionInternal()
 {
 	FContext* const Context = FContext::Get();
 	AUTORTFM_ASSERT(Context->GetStatus() == EContextStatus::OnTrack);
-	Context->AbortTransactionWithPostAbortCallback(EContextStatus::AbortedByCascadingAbort, std::move(RunAfterAbort));
+	Context->AbortTransaction(EContextStatus::AbortedByCascadingAbort);
 }
 
-AUTORTFM_INTERNAL void RTFM_CascadingRetryTransactionInternal(TTask<void()> && RunAfterAbortBeforeRetryWork)
+AUTORTFM_INTERNAL void RTFM_CascadingRetryTransactionInternal()
 {
 	FContext* Context = FContext::Get();
 	AUTORTFM_ASSERT(Context->GetStatus() == EContextStatus::OnTrack);
-	Context->AbortTransactionWithPostAbortCallback(EContextStatus::AbortedByCascadingRetry, std::move(RunAfterAbortBeforeRetryWork));
+	Context->AbortTransaction(EContextStatus::AbortedByCascadingRetry);
 }
 
 AUTORTFM_INTERNAL void RTFM_OnCommitInternal(TTask<void()> && Work)
@@ -1052,17 +1052,17 @@ void ForTheRuntime::Initialize(const ForTheRuntime::FExternAPI& ExternAPI)
 	autortfm_initialize(&ExternAPI);
 }
 
-void ForTheRuntime::CascadingAbortTransactionInternal(TTask<void()>&& RunAfterAbort)
+void ForTheRuntime::CascadingAbortTransactionInternal()
 {
 	AUTORTFM_ASSERT(autortfm_is_closed());
-	RTFM_CascadingAbortTransactionInternal(std::move(RunAfterAbort));
+	RTFM_CascadingAbortTransactionInternal();
 }
 
-void ForTheRuntime::CascadingRetryTransactionInternal(TTask<void()>&& RunAfterAbortBeforeRetryWork)
+void ForTheRuntime::CascadingRetryTransactionInternal()
 {
 	if (autortfm_is_closed())
 	{
-		return RTFM_CascadingRetryTransactionInternal(std::move(RunAfterAbortBeforeRetryWork));
+		return RTFM_CascadingRetryTransactionInternal();
 	}
 }
 
diff --git a/Engine/Source/Runtime/AutoRTFM/Private/Context.cpp b/Engine/Source/Runtime/AutoRTFM/Private/Context.cpp
--- a/Engine/Source/Runtime/AutoRTFM/Private/Context.cpp
+++ b/Engine/Source/Runtime/AutoRTFM/Private/Context.cpp
@@ -206,27 +206,6 @@ void FContext::AbortTransaction(EContextStatus NewStatus)
 	Throw();
 }
 
-void FContext::AbortTransactionWithPostAbortCallback(EContextStatus NewStatus, TTask<void()>&& InCallback)
-{
-	// The callback parameter is only honored by the cascading aborts.
-	AUTORTFM_ASSERT(NewStatus == EContextStatus::AbortedByCascadingAbort ||
-		   NewStatus == EContextStatus::AbortedByCascadingRetry);
-
-	// We must explicitly copy the passed-in callback here, because the original may have been
-	// allocated within a transactional context; if so, its memory was allocated under
-	// transactionalized conditions. By copying, we create an open copy of the callback.
-	if (!PostAbortCallback.IsSet())
-	{
-		PostAbortCallback = InCallback;
-	}
-	else if (InCallback.IsSet())
-	{
-		AUTORTFM_FATAL("Only one post-abort callback at a time is supported");
-	}
-
-	AbortTransaction(NewStatus);
-}
-
 bool FContext::IsAborting() const
 {
 	switch (Status)
@@ -543,7 +522,6 @@ ETransactionResult FContext::Transact(void (*UninstrumentedFunction)(void*), voi
 				break;
 
 			case EContextStatus::AbortedByRequest:
-				AUTORTFM_ASSERT(!PostAbortCallback.IsSet());
 				Result = ETransactionResult::AbortedByRequest;
 				break;
 
@@ -552,39 +530,15 @@ ETransactionResult FContext::Transact(void (*UninstrumentedFunction)(void*), voi
 				break;
 
 			case EContextStatus::AbortedByCascadingAbort:
-				if (PostAbortCallback.IsSet())
-				{
-					// Call the post-abort callback to do whatever work the user
-					// required be done before throwing.
-					Status = EContextStatus::InPostAbort;
-
-					PostAbortCallback();
-					PostAbortCallback.Reset();
-
-					AUTORTFM_ASSERT(Status == EContextStatus::InPostAbort);
-					Status = EContextStatus::AbortedByCascadingAbort;
-				}
-
 				Result = ETransactionResult::AbortedByCascade;
 				break;
 
 			case EContextStatus::AbortedByCascadingRetry:
-				AUTORTFM_ASSERT(PostAbortCallback.IsSet());
-			
 				// Clean up the transaction to get it ready for re-execution.
 				ClearTransactionStatus();
 				CurrentTransaction->Reset();
 
 				AUTORTFM_ASSERT(Status == EContextStatus::OnTrack);
-				Status = EContextStatus::InPostAbort;
-
-				// Call the post-abort callback to do whatever work the user
-				// required be done before we attempt to re-execute.
-				PostAbortCallback();
-				PostAbortCallback.Reset();
-				
-				AUTORTFM_ASSERT(Status == EContextStatus::InPostAbort);
-				Status = EContextStatus::OnTrack;
 
 				// Then get rolling!
 				CurrentTransaction->SetClosedActive();
diff --git a/Engine/Source/Runtime/AutoRTFM/Private/Context.h b/Engine/Source/Runtime/AutoRTFM/Private/Context.h
--- a/Engine/Source/Runtime/AutoRTFM/Private/Context.h
+++ b/Engine/Source/Runtime/AutoRTFM/Private/Context.h
@@ -43,7 +43,6 @@ public:
 
 	ETransactionResult CommitTransaction();
 	[[noreturn]] void AbortTransaction(EContextStatus NewStatus);
-	[[noreturn]] void AbortTransactionWithPostAbortCallback(EContextStatus NewStatus, TTask<void()>&& InCallback);
 	void ClearTransactionStatus();
 	bool IsAborting() const;
 	bool IsTransactional() const;
@@ -176,7 +175,6 @@ private:
     EContextStatus Status{EContextStatus::Idle};
 	FThreadID CurrentThreadId;
 	uint32_t StackLocalInitializerDepth = 0;
-	TTask<void()> PostAbortCallback;
 	FTaskPool TaskPool;
 	FTransactionPool TransactionPool;
 	FOpenHashThrottler OpenHashThrottler{
diff --git a/Engine/Source/Runtime/AutoRTFM/Private/ContextStatus.h b/Engine/Source/Runtime/AutoRTFM/Private/ContextStatus.h
--- a/Engine/Source/Runtime/AutoRTFM/Private/ContextStatus.h
+++ b/Engine/Source/Runtime/AutoRTFM/Private/ContextStatus.h
@@ -34,8 +34,6 @@ inline const char* GetContextStatusName(EContextStatus Status)
         return "AbortedByCascadingRetry";
     case EContextStatus::InStaticLocalInitializer:
         return "InStaticLocalInitializer";
-    case EContextStatus::InPostAbort:
-        return "InPostAbort";
     }
 }
 
diff --git a/Engine/Source/Runtime/AutoRTFM/Private/Utils.cpp b/Engine/Source/Runtime/AutoRTFM/Private/Utils.cpp
--- a/Engine/Source/Runtime/AutoRTFM/Private/Utils.cpp
+++ b/Engine/Source/Runtime/AutoRTFM/Private/Utils.cpp
@@ -95,17 +95,23 @@ UE_AUTORTFM_API UE_AUTORTFM_FORCENOINLINE void ReportError(const char* File, int
 		}
 	}
 
-	if (FContext* const Context = FContext::Get())
+	switch (InternalAbortAction)
 	{
-		if (InternalAbortAction == ForTheRuntime::EAutoRTFMInternalAbortActionState::Abort)
-		{
-			Context->AbortByLanguageAndThrow();
-		}
-		else
-		{
-			TTask<void()> Task = [] { AUTORTFM_ENSURE(ForTheRuntime::SetAutoRTFMRuntime(ForTheRuntime::AutoRTFM_ForcedDisabled)); };
-			Context->AbortTransactionWithPostAbortCallback(EContextStatus::AbortedByCascadingRetry, std::move(Task));
-		}
+		case ForTheRuntime::EAutoRTFMInternalAbortActionState::Abort:
+			if (FContext* const Context = FContext::Get())
+			{
+				Context->AbortByLanguageAndThrow();
+			}
+			break;
+
+		case ForTheRuntime::EAutoRTFMInternalAbortActionState::Crash:
+			// The `autortfm_log_fatal` lines above probably triggered a crash already, but if we made it this far, force the crash.
+			__builtin_trap();
+			break;
+
+		default:
+			::AutoRTFM::EnsureFailure(File, Line, ProgramCounter, "Unexpected InternalAbortAction", "InternalAbortAction = %d", int(InternalAbortAction));
+			break;
 	}
 }
 
diff --git a/Engine/Source/Runtime/AutoRTFM/Public/AutoRTFM.h b/Engine/Source/Runtime/AutoRTFM/Public/AutoRTFM.h
--- a/Engine/Source/Runtime/AutoRTFM/Public/AutoRTFM.h
+++ b/Engine/Source/Runtime/AutoRTFM/Public/AutoRTFM.h
@@ -498,11 +498,6 @@ enum class EContextStatus : uint8_t
 	// Means we are in a static local initializer which always run in the open.
 	// `IsTransactional()` will return `false` when in this state.
 	InStaticLocalInitializer = autortfm_status_in_static_local_initializer,
-
-	// An InPostAbort status means that a transaction has just been aborted, but we have not yet
-	// thrown back to the caller or retried the transaction. We are only in this state when 
-	// executing the post-abort callback.
-	InPostAbort = autortfm_status_in_post_abort,
 };
 
 // An enumerator of transactional memory validation levels.
@@ -529,8 +524,8 @@ namespace ForTheRuntime
 {
 	using FExternAPI = autortfm_extern_api;
 	UE_AUTORTFM_API void Initialize(const FExternAPI& ExternAPI);
-	UE_AUTORTFM_API void CascadingAbortTransactionInternal(TTask<void()>&& RunAfterAbort);
-	UE_AUTORTFM_API void CascadingRetryTransactionInternal(TTask<void()>&& RunAfterAbortBeforeRetryWork);
+	UE_AUTORTFM_API void CascadingAbortTransactionInternal();
+	UE_AUTORTFM_API void CascadingRetryTransactionInternal();
 	UE_AUTORTFM_API void OnCommitInternal(TTask<void()>&& Work);
 	UE_AUTORTFM_API void OnAbortInternal(TTask<void()>&& Work);
 	UE_AUTORTFM_API void OnCompleteInternal(TTask<void()>&& Work);
@@ -651,24 +646,16 @@ UE_AUTORTFM_CRITICAL_INLINE void AbortTransaction()
 
 // End a transaction nest in the closed, discarding all effects. This cascades, 
 // meaning an abort of a nested transaction will cause all transactions in the
-// nest to abort. Once the transaction has aborted, `RunAfterAbort` will
-// be called. Finally, control will be returned to the end of the outermost 
+// nest to abort. Once the transaction has aborted, on-complete callbacks will
+// be invoked. Finally, control will be returned to the end of the outermost 
 // Transact.
 #if UE_AUTORTFM
-template <typename FunctorType>
-UE_AUTORTFM_CRITICAL_INLINE void CascadingAbortTransaction(FunctorType&& RunAfterAbort)
-{
-	ForTheRuntime::CascadingAbortTransactionInternal(std::forward<FunctorType>(RunAfterAbort));
-}
-
 UE_AUTORTFM_CRITICAL_INLINE void CascadingAbortTransaction()
 {
-	ForTheRuntime::CascadingAbortTransactionInternal({});
+	ForTheRuntime::CascadingAbortTransactionInternal();
 }
 #else
-template <typename FunctorType>
-UE_AUTORTFM_CRITICAL_INLINE void CascadingAbortTransaction(FunctorType&&) {}
-UE_AUTORTFM_CRITICAL_INLINE void CascadingAbortTransaction()              {}
+UE_AUTORTFM_CRITICAL_INLINE void CascadingAbortTransaction() {}
 #endif
 
 namespace Detail
@@ -866,20 +853,19 @@ template<typename FunctorType> [[nodiscard]] UE_AUTORTFM_CRITICAL_INLINE EContex
 }
 
 // Force a transaction nest to be retried. Once the transaction has aborted,
-// it will call `RunAfterAbortBeforeRetryWork`, before retrying the
-// transaction. This is an expensive operation and should thus be used with
-// extreme caution. If this is called outside of a transaction then the
-// `RunAfterAbortBeforeRetryWork` callback is never executed.
+// all on-complete handlers will be called before retrying the transaction. 
+// This is an expensive operation and should thus be used with extreme caution.
+// If this is called outside of a transaction, the call is ignored.
 #if UE_AUTORTFM
-template<typename FunctorType> UE_AUTORTFM_CRITICAL_INLINE void CascadingRetryTransaction(FunctorType&& RunAfterAbortBeforeRetryWork)
+UE_AUTORTFM_CRITICAL_INLINE void CascadingRetryTransaction()
 {
 	if (autortfm_is_closed())
 	{
-		ForTheRuntime::CascadingRetryTransactionInternal(std::forward<FunctorType>(RunAfterAbortBeforeRetryWork));
+		ForTheRuntime::CascadingRetryTransactionInternal();
 	}
 }
 #else
-template<typename FunctorType> UE_AUTORTFM_CRITICAL_INLINE void CascadingRetryTransaction(FunctorType&&)
+UE_AUTORTFM_CRITICAL_INLINE void CascadingRetryTransaction()
 {
 }
 #endif
diff --git a/Engine/Source/Runtime/CoreUObject/Private/Serialization/TransactionallySafeAsyncLoading.cpp b/Engine/Source/Runtime/CoreUObject/Private/Serialization/TransactionallySafeAsyncLoading.cpp
--- a/Engine/Source/Runtime/CoreUObject/Private/Serialization/TransactionallySafeAsyncLoading.cpp
+++ b/Engine/Source/Runtime/CoreUObject/Private/Serialization/TransactionallySafeAsyncLoading.cpp
@@ -300,7 +300,7 @@ public:
 							AutoRTFM::PopOnAbortHandler(this);
 						}
 
-						AutoRTFM::CascadingRetryTransaction([this, RequestIdsClone = ::MoveTemp(RequestIdsClone)]() mutable
+						AutoRTFM::OnComplete([this, RequestIdsClone = ::MoveTemp(RequestIdsClone)]() mutable
 						{
 							UE_LOG(LogStreaming, Display, TEXT("A call to `FlushLoading` that is flushing non-cached packages is causing Verse to abort and retry."));
 							DrainTransactionallyDeferredLoadPackages(false, &RequestIdsClone, true);
@@ -308,6 +308,7 @@ public:
 							WrappedPackageLoader->FlushLoading(RequestIdsClone);
 							UE_LOG(LogStreaming, Display, TEXT("`FlushLoading` has completed after Verse aborted, and we are now retrying."));
 						});
+						AutoRTFM::CascadingRetryTransaction();
 					});
 					ensure(Status == AutoRTFM::EContextStatus::AbortedByCascadingRetry);
 				}
diff --git a/Engine/Source/Runtime/Engine/Private/Components/SkeletalMeshComponent.cpp b/Engine/Source/Runtime/Engine/Private/Components/SkeletalMeshComponent.cpp
--- a/Engine/Source/Runtime/Engine/Private/Components/SkeletalMeshComponent.cpp
+++ b/Engine/Source/Runtime/Engine/Private/Components/SkeletalMeshComponent.cpp
@@ -4776,19 +4776,20 @@ bool USkeletalMeshComponent::HandleExistingParallelEvaluationTask(bool bBlockOnT
 			// If we try and complete existing parallel evaluation in a transaction, we have to abort the transaction
 			// and retry up to this point. This is because we can only have started parallel evaluation outside of a
 			// transaction, and thus it is only safe to complete the work outside of a transaction too! For instance
-			// if the paralle animation evaluation task depended on something that was trying to take a mutex that we
+			// if the parallel animation evaluation task depended on something that was trying to take a mutex that we
 			// had also taken in the transaction, we'd have caused a deadlock if we didn't have this here.
 			if (AutoRTFM::IsTransactional())
 			{
 				(void)AutoRTFM::Close([&]
 				{
-					AutoRTFM::CascadingRetryTransaction([strong_this = TObjectPtr<USkeletalMeshComponent>(this), bPerformPostAnimEvaluation]
+					AutoRTFM::OnComplete([strong_this = TObjectPtr<USkeletalMeshComponent>(this), bPerformPostAnimEvaluation]
 					{
 						UE_LOG(LogSkeletalMesh, Display, TEXT("A call to `HandleExistingParallelEvaluationTask` that is blocking is causing Verse to abort and retry."));
 						const bool bBlockOnTask = true; // Was true from the if block above!
 						strong_this->HandleExistingParallelEvaluationTask(bBlockOnTask, bPerformPostAnimEvaluation);
 						UE_LOG(LogSkeletalMesh, Display, TEXT("`HandleExistingParallelEvaluationTask` has completed after Verse aborted, and we are now retrying."));
 					});
+					AutoRTFM::CascadingRetryTransaction();
 				});
 				return true;
 			}
