[AutoRTFM] Add new & delete external mappings when the ANSI allocator is enabled

Is used when the sanitizers are enabled (ASAN, TSAN, etc)

diff --git a/Engine/Source/Runtime/AutoRTFM/AutoRTFM.Build.cs b/Engine/Source/Runtime/AutoRTFM/AutoRTFM.Build.cs
--- a/Engine/Source/Runtime/AutoRTFM/AutoRTFM.Build.cs
+++ b/Engine/Source/Runtime/AutoRTFM/AutoRTFM.Build.cs
@@ -34,6 +34,18 @@ namespace UnrealBuildTool.Rules
 					AutoRTFMExternalMappingFiles.Add("Runtime/AutoRTFM/Public/ASAN.aem");
 				}
 			}
+			if (Target.Platform == UnrealTargetPlatform.Linux)
+			{
+				AutoRTFMExternalMappingFiles.Add("Runtime/AutoRTFM/Public/StdLib.Linux.aem");
+				if (Target.LinuxPlatform.bEnableAddressSanitizer)
+				{
+					AutoRTFMExternalMappingFiles.Add("Runtime/AutoRTFM/Public/ASAN.aem");
+				}
+				if (Target.StaticAllocator == StaticAllocatorType.Ansi)
+				{
+					AutoRTFMExternalMappingFiles.Add("Runtime/AutoRTFM/Public/NewDelete.Linux.aem");
+				}
+			}
 
 			// Enable (almost) all warnings as errors
 			CppCompileWarningSettings.ShadowVariableWarningLevel = WarningLevel.Error;
