[AutoRTFM] Remove assertion in FTransaction::ShouldRecordWrite()

The assertion caught transactional writes to stack memory that could be observed without reverting to the original value post-abort. However, this could fire for patterns where this was technically safe. Work-arounds for this assertion proved either too inefficient or severely hurt code readability.

Given that unscoped transactions would require the implicit `RecordWrite()` call to be inside an `AutoRTFM::Close()` call, and that function is annotated with `[[nodiscard]]`, it does seem unlikely for a caller to depend on stack variables modified in the closed block when the transaction has been aborted - most callers should immediately return or immediately bubble the abort.

It was also noted by [REDACTED_PERSON] that the assertion should have applied to both scoped and unscoped transactions, as an abort jumps to the nearest `Close()` or `Transact()`, not just `Transact()` as the comment suggested.

For these reasons, I've reluctantly conceded that this check is being less helpful than the additional safety it provided.

diff --git a/Engine/Source/Runtime/AutoRTFM/Private/Context.cpp b/Engine/Source/Runtime/AutoRTFM/Private/Context.cpp
--- a/Engine/Source/Runtime/AutoRTFM/Private/Context.cpp
+++ b/Engine/Source/Runtime/AutoRTFM/Private/Context.cpp
@@ -221,8 +221,6 @@ bool FContext::IsAborting() const
 
 EContextStatus FContext::CallClosedNest(void (*ClosedFunction)(void* Arg), void* Arg)
 {
-	TScopedGuard<void*> ClosedStackAddressGuard(ClosedStackAddress, &ClosedStackAddressGuard);
-
 	FTransaction* const Transaction = GetCurrentTransaction();
 	AUTORTFM_ASSERT(Transaction != nullptr);
 	AUTORTFM_ASSERT(FTransaction::EState::OpenActive == Transaction->State());
diff --git a/Engine/Source/Runtime/AutoRTFM/Private/Context.h b/Engine/Source/Runtime/AutoRTFM/Private/Context.h
--- a/Engine/Source/Runtime/AutoRTFM/Private/Context.h
+++ b/Engine/Source/Runtime/AutoRTFM/Private/Context.h
@@ -89,12 +89,6 @@ public:
 	inline EContextStatus GetStatus() const { return CurrentThreadId == FThreadID::GetCurrent() ? Status : EContextStatus::Idle; }
 	[[noreturn]] void Throw();
 
-	// Returns the starting stack address of the innermost call to Closed(), or
-	// nullptr if there is no call to Closed. Used to assert that a stack memory
-	// write is safe to record.
-	// See FTransaction::ShouldRecordWrite()
-	inline const void* GetClosedStackAddress() const { return ClosedStackAddress; }
-
 	inline FOpenHashThrottler& GetOpenHashThrottler() { return OpenHashThrottler; }
 
 	void DumpState() const;
@@ -171,7 +165,6 @@ private:
 	FCallNest* CurrentNest{nullptr};
 
     FStackRange Stack;
-	void* ClosedStackAddress = nullptr;
     EContextStatus Status{EContextStatus::Idle};
 	FThreadID CurrentThreadId;
 	uint32_t StackLocalInitializerDepth = 0;
diff --git a/Engine/Source/Runtime/AutoRTFM/Private/TransactionInlines.h b/Engine/Source/Runtime/AutoRTFM/Private/TransactionInlines.h
--- a/Engine/Source/Runtime/AutoRTFM/Private/TransactionInlines.h
+++ b/Engine/Source/Runtime/AutoRTFM/Private/TransactionInlines.h
@@ -53,50 +53,13 @@ UE_AUTORTFM_FORCEINLINE bool FTransaction::IsOnStack(const void* LogicalAddress)
 
 UE_AUTORTFM_FORCEINLINE bool FTransaction::ShouldRecordWrite(void* LogicalAddress) const
 {
-    // We cannot record writes to stack memory used within the transaction, as
-    // undoing the writes may corrupt stack memory that has been unwound or
-    // is now being used for a different variable from the one the write was
-    // made.
-    if (!IsOnStack(LogicalAddress))
-    {
-        return true;
-    }
-
-    // Writes to the stack under a scoped-transaction can be safely ignored,
-    // because the values on the stack are not visible outside of the scope of
-    // the transaction. In other words, if a scoped-transaction aborts that
-    // memory will cease to be meaningful anyway.
-
-    // Non-scoped transactions, as the name implies, do not impose a lexical
-    // scope that encompasses the transaction. Instead a non-scoped transaction
-    // is started with a call to StartTransaction() and ended with a call to
-    // either AbortTransaction() or CommitTransaction(). Unlike a
-    // scoped transaction, there's no precise stack range for a non-scoped
-    // transaction, as the scope can freely grow or shrink between the calls to
-    // [Start|Abort|Commit]Transaction() and any recorded writes. The only
-    // guarantee we have is that a non-scoped transaction cannot shrink past the
-    // outer scoped transaction. For this reason, non-scoped transactions
-    // adopt the stack range of the outer transaction, as this is guaranteed
-    // to encompass the non-scoped transaction's scope range.
-
-    // For non-scoped transactions, we assert that we're not writing to a
-    // memory address that's in the transaction's stack range as this cannot be
-    // safely undone, and stack variables may be visible once the transaction is
-    // aborted. We make an exception for stack variables declared within the
-    // scope of a Close(), as writing to these stack variables can be safely
-    // ignored (they have the same constrained visibility as stack variables in
-    // a scoped transaction).
-
-    // Hitting this assert?
-    // Consider moving the variable being written to an inner scoped
-    // transaction, or move the variable outside of the nearest parent
-    // scoped-transaction.
-    AUTORTFM_ASSERT(bIsStackScoped || LogicalAddress < FContext::Get()->GetClosedStackAddress());
-
-    return false;
+	// We cannot record writes to stack memory used within the transaction, as
+	// undoing the writes may corrupt stack memory that has been unwound or
+	// is now being used for a different variable from the one the write was
+	// made.
+	return !IsOnStack(LogicalAddress);
 }
 
-
 AUTORTFM_NO_ASAN UE_AUTORTFM_FORCEINLINE
 void FTransaction::RecordWrite(void* LogicalAddress, size_t Size, bool bNoMemoryValidation /* = false */)
 {
