[AutoRTFM]
* Added HEADER_UNIT_SKIP to prevent some headers from being compiled stand alone (because they cause unused-warnings that are enabled by AutoRTFM)

diff --git a/Engine/Source/Runtime/AutoRTFM/Private/Toggles.h b/Engine/Source/Runtime/AutoRTFM/Private/Toggles.h
--- a/Engine/Source/Runtime/AutoRTFM/Private/Toggles.h
+++ b/Engine/Source/Runtime/AutoRTFM/Private/Toggles.h
@@ -2,6 +2,8 @@
 
 #pragma once
 
+// HEADER_UNIT_SKIP - unused warnings
+
 #include "BuildMacros.h"
 
 namespace AutoRTFM
diff --git a/Engine/Source/Runtime/AutoRTFM/Public/AutoRTFM.h b/Engine/Source/Runtime/AutoRTFM/Public/AutoRTFM.h
--- a/Engine/Source/Runtime/AutoRTFM/Public/AutoRTFM.h
+++ b/Engine/Source/Runtime/AutoRTFM/Public/AutoRTFM.h
@@ -2,6 +2,8 @@
 
 #pragma once
 
+// HEADER_UNIT_SKIP - unused warnings
+
 #include "AutoRTFMConstants.h"
 
 #include "AutoRTFMDefines.h"
