[AutoRTFM] Implement pooling for FCallNest objects.

A call nest is allocated every time we go into the closed. By pooling them, we save an allocation each time we transition from the open into the closed.

This slightly improves many benchmarks, but the improvement is most visible in `Benchmarks.TopLevelTransaction` which specifically measures the cost of entering a transaction.

Before CL:

benchmark name                       samples       iterations    estimated
mean          low mean      high mean
std dev       low std dev   high std dev
-------------------------------------------------------------------------------
Commit                                         [REDACTED_PHONE] ms
933.921 us    931.658 us    936.544 us
12.4079 us    10.4639 us    15.3629 us

Abort                                          [REDACTED_PHONE] ms
886.057 us     883.52 us    889.118 us
14.1162 us    11.8403 us    17.6558 us

After CL:
-------------------------------------------------------------------------------
Commit                                         [REDACTED_PHONE] ms
713.402 us    709.953 us     719.56 us
22.8077 us    14.7729 us    41.6132 us

Abort                                          [REDACTED_PHONE] ms
676.898 us    674.336 us    680.026 us
14.4736 us    12.0771 us    20.0404 us

diff --git a/Engine/Source/Runtime/AutoRTFM/Private/Context.cpp b/Engine/Source/Runtime/AutoRTFM/Private/Context.cpp
--- a/Engine/Source/Runtime/AutoRTFM/Private/Context.cpp
+++ b/Engine/Source/Runtime/AutoRTFM/Private/Context.cpp
@@ -228,7 +228,7 @@ EContextStatus FContext::CallClosedNest(void (*ClosedFunction)(void* Arg), void*
 	const void* PreviousOpenReturnAddress = Transaction->OpenReturnAddress();
 	Transaction->SetClosedActive();
 
-	PushCallNest(new FCallNest(this));
+	PushCallNest(CallNestPool.Take(this));
 
 	CurrentNest->Try([&]() { ClosedFunction(Arg); });
 
@@ -257,7 +257,7 @@ void FContext::PopCallNest()
 	FCallNest* OldCallNest = CurrentNest;
 	CurrentNest = CurrentNest->Parent;
 
-	delete OldCallNest;
+	CallNestPool.Return(OldCallNest);
 }
 
 FTransaction* FContext::PushTransaction(
@@ -413,7 +413,7 @@ ETransactionResult FContext::Transact(void (*UninstrumentedFunction)(void*), voi
 	MaterializeDeferredTransactions();
 	AUTORTFM_ASSERT(!GetNumDeferredTransactions());
     
-	FCallNest* NewNest = new FCallNest(this);
+	FCallNest* NewNest = CallNestPool.Take(this);
 
 	void* TransactStackStart = &NewNest;
 
diff --git a/Engine/Source/Runtime/AutoRTFM/Private/Context.h b/Engine/Source/Runtime/AutoRTFM/Private/Context.h
--- a/Engine/Source/Runtime/AutoRTFM/Private/Context.h
+++ b/Engine/Source/Runtime/AutoRTFM/Private/Context.h
@@ -4,8 +4,10 @@
 
 #if (defined(__AUTORTFM) && __AUTORTFM)
 
+#include "CallNest.h"
 #include "IntrusivePool.h"
 #include "OpenHashThrottler.h"
+#include "Pool.h"
 #include "StackRange.h"
 #include "TaskArray.h"
 #include "ThreadID.h"
@@ -15,13 +17,11 @@
 namespace AutoRTFM
 {
 
-class FLineLock;
-class FCallNest;
-
 class AUTORTFM_DISABLE FContext final
 {
 	using FTaskPool = typename TTaskArray<TTask<void()>>::FEntryPool;
 	using FTransactionPool = TIntrusivePool<FTransaction, 16>;
+	using FCallNestPool = TPool<FCallNest, 16>;
 
 public:
 	static FContext* Create();
@@ -170,6 +170,7 @@ private:
 	uint32_t StackLocalInitializerDepth = 0;
 	FTaskPool TaskPool;
 	FTransactionPool TransactionPool;
+	FCallNestPool CallNestPool;
 	FOpenHashThrottler OpenHashThrottler{
 		/* LogInterval            */ 10,   // Log every 10 seconds.
 		/* AdjustThrottleInterval */ 0.5,  // Adjust throttling probabilities every 500ms.
diff --git a/Engine/Source/Runtime/AutoRTFM/Private/Transaction.h b/Engine/Source/Runtime/AutoRTFM/Private/Transaction.h
--- a/Engine/Source/Runtime/AutoRTFM/Private/Transaction.h
+++ b/Engine/Source/Runtime/AutoRTFM/Private/Transaction.h
@@ -203,7 +203,12 @@ private:
 
 	FContext* Context;
 
-	// Are we nested? Then this is the parent.
+	// This field is used for two purposes:
+	// - When the transaction is active, the Parent field is used to support transaction nesting.
+	//   This field points to our immediate outer nest; it's null when this is the top level.
+	// - When the transaction is "freed"--sitting in the FTransactionPool free list--this field is
+	//   used to point to the next free item, forming a linked list of reusable FTransactions. 
+	//   (See GetIntrusiveAddress.)
 	FTransaction* Parent;
 
 	// Commit tasks run on commit in forward order.
