Fix a performance regression in AutoRTFM (up to 1.86x slower in worst case without this).

The issue was caused by poor codegen with `[[autortfm(disabled}]]` functions where the compiler would leave in stack regions it could have elided.

#rnx

diff --git a/Engine/Source/Runtime/AutoRTFM/Private/CompilerInterface.cpp b/Engine/Source/Runtime/AutoRTFM/Private/CompilerInterface.cpp
--- a/Engine/Source/Runtime/AutoRTFM/Private/CompilerInterface.cpp
+++ b/Engine/Source/Runtime/AutoRTFM/Private/CompilerInterface.cpp
@@ -151,12 +151,12 @@ extern "C" AUTORTFM_DISABLE UE_AUTORTFM_API void autortfm_llvm_missing_function(
     Context->AbortByLanguageAndThrow();
 }
 
-extern "C" AUTORTFM_DISABLE UE_AUTORTFM_API void autortfm_called_no_autortfm() noexcept
+extern "C" AUTORTFM_DISABLE UE_AUTORTFM_FORCENOINLINE UE_AUTORTFM_API void autortfm_called_no_autortfm() noexcept
 {
 	AUTORTFM_FATAL("inlined UE_AUTORTFM_NOAUTORTFM function called from the closed");
 }
 
-extern "C" AUTORTFM_DISABLE UE_AUTORTFM_API void autortfm_disable_called() noexcept
+extern "C" AUTORTFM_DISABLE UE_AUTORTFM_FORCENOINLINE UE_AUTORTFM_API void autortfm_disable_called() noexcept
 {
 	AUTORTFM_FATAL("inlined AUTORTFM_DISABLE function called from the closed");
 }
diff --git a/Engine/Source/Runtime/AutoRTFM/Private/Context.h b/Engine/Source/Runtime/AutoRTFM/Private/Context.h
--- a/Engine/Source/Runtime/AutoRTFM/Private/Context.h
+++ b/Engine/Source/Runtime/AutoRTFM/Private/Context.h
@@ -79,11 +79,6 @@ public:
 	{
 		return CurrentTransaction;
 	}
-	inline FTransaction* GetCurrentTransactionAssumingNoMaterialization()
-	{
-		AUTORTFM_ASSERT(!MustMaterializeDeferredTransactions());
-		return CurrentTransaction;
-	}
 
 	inline FCallNest* GetCurrentNest() const { return CurrentNest; }
 	inline EContextStatus GetStatus() const { return CurrentThreadId == FThreadID::GetCurrent() ? Status : EContextStatus::Idle; }
diff --git a/Engine/Source/Runtime/AutoRTFM/Private/ContextInlines.h b/Engine/Source/Runtime/AutoRTFM/Private/ContextInlines.h
--- a/Engine/Source/Runtime/AutoRTFM/Private/ContextInlines.h
+++ b/Engine/Source/Runtime/AutoRTFM/Private/ContextInlines.h
@@ -28,7 +28,7 @@ template<unsigned SIZE> AUTORTFM_NO_ASAN UE_AUTORTFM_FORCEINLINE void FContext::
 	}
 	else
 	{
-		GetCurrentTransactionAssumingNoMaterialization()->RecordWrite<SIZE>(LogicalAddress);
+		GetMaterializedTransaction()->RecordWrite<SIZE>(LogicalAddress);
 	}
 }
 
