[AutoRTFM] Register external mappings into the function table

The AutoRTFM External Mappings Files should cover all the mappings provided by UE_AUTORTFM_REGISTER_OPEN_TO_CLOSED_FUNCTIONS() in StdLibSupport.cpp, but to be extra cautious, turn these mappings into a fallback lookup which will ensure. Once we're confident there are no gaps in the AEM files, we can remove this fallback path.

diff --git a/Engine/Source/Runtime/AutoRTFM/Private/API.cpp b/Engine/Source/Runtime/AutoRTFM/Private/API.cpp
--- a/Engine/Source/Runtime/AutoRTFM/Private/API.cpp
+++ b/Engine/Source/Runtime/AutoRTFM/Private/API.cpp
@@ -726,12 +726,18 @@ AUTORTFM_INTERNAL void ProcessAllPendingOpenToClosedRegistrations()
 
 }  // anonymous namespace
 
+// Populated by the AutoRTFM compiler
+extern "C" autortfm_open_to_closed_table autortfm_external_mappings;
+
 // The AutoRTFM public API.
 
 extern "C" void autortfm_initialize(const autortfm_extern_api* ExternAPI) noexcept
 {
 	AutoRTFM::UnreachableIfClosed("TODO: Mark this as AUTORTFM_DISABLED and fixup callers");
 	AUTORTFM_ENSURE_MSG(!FContext::Get(), "AutoRTFM initialized twice");
+
+	autortfm_register_open_to_closed_functions(&autortfm_external_mappings);
+
 	AUTORTFM_ASSERT(ExternAPI);
 	AUTORTFM_ASSERT(ExternAPI->Allocate);
 	AUTORTFM_ASSERT(ExternAPI->AllocateZeroed);
diff --git a/Engine/Source/Runtime/AutoRTFM/Private/CompilerInterface.cpp b/Engine/Source/Runtime/AutoRTFM/Private/CompilerInterface.cpp
--- a/Engine/Source/Runtime/AutoRTFM/Private/CompilerInterface.cpp
+++ b/Engine/Source/Runtime/AutoRTFM/Private/CompilerInterface.cpp
@@ -87,15 +87,7 @@ extern "C" AUTORTFM_DISABLE UE_AUTORTFM_API AUTORTFM_NO_ASAN void autortfm_recor
 	}
 }
 
-// Mark autortfm_lookup_function() with UE_AUTORTFM_NOAUTORTFM and use
-// UE_AUTORTFM_MAP_OPEN_TO_SELF to map the closed function to the open.
-// We do this to prevent calls to autortfm_lookup_function() being wrapped
-// with calls to autortfm_pre_open() and autortfm_post_open(), which is
-// expensive for hot code like this.
-UE_AUTORTFM_REGISTER_OPEN_TO_CLOSED_FUNCTIONS(
-	UE_AUTORTFM_MAP_OPEN_TO_SELF(autortfm_lookup_function)
-);
-extern "C" UE_AUTORTFM_NOAUTORTFM UE_AUTORTFM_API void* autortfm_lookup_function(void* OriginalFunction, const char* Where) noexcept
+extern "C" UE_AUTORTFM_ALWAYS_OPEN UE_AUTORTFM_API void* autortfm_lookup_function(void* OriginalFunction, const char* Where) noexcept
 {
 	return FunctionMapLookup(OriginalFunction, Where);
 }
diff --git a/Engine/Source/Runtime/AutoRTFM/Private/FunctionMap.cpp b/Engine/Source/Runtime/AutoRTFM/Private/FunctionMap.cpp
--- a/Engine/Source/Runtime/AutoRTFM/Private/FunctionMap.cpp
+++ b/Engine/Source/Runtime/AutoRTFM/Private/FunctionMap.cpp
@@ -7,6 +7,7 @@
 #include "AutoRTFM.h"
 #include "BuildMacros.h"
 #include "FunctionMapInlines.h"
+#include "StdLibSupport.h"
 #include "Utils.h"
 
 #include <cinttypes>
@@ -329,6 +330,13 @@ void* FunctionMapLookupExhaustive(void* OpenFn, const char* Where)
 		}
 	}
 
+	// Check the temporary, old standard library mappings.
+	// This will be removed in the near future.
+	if (void* ClosedFn = LookupStdLibFunctionAndEnsureIfFound(OpenFn))
+	{
+		return ClosedFn;
+	}
+
 	AUTORTFM_MUST_TAIL return FunctionMapReportError(OpenFn, Where);
 }
 
diff --git a/Engine/Source/Runtime/AutoRTFM/Private/StdLibSupport.cpp b/Engine/Source/Runtime/AutoRTFM/Private/StdLibSupport.cpp
--- a/Engine/Source/Runtime/AutoRTFM/Private/StdLibSupport.cpp
+++ b/Engine/Source/Runtime/AutoRTFM/Private/StdLibSupport.cpp
@@ -2,6 +2,7 @@
 
 #if (defined(__AUTORTFM) && __AUTORTFM)
 
+#include "StdLibSupport.h"
 #include "AutoRTFM.h"
 #include "BuildMacros.h"
 #include "ContextInlines.h"
@@ -990,10 +991,17 @@ extern "C" AUTORTFM_INTERNAL UE_AUTORTFM_API void RTFM_ZdaPvmSt11align_val_t(voi
 // ^^^^^^^^^^^^^^^^^^^^^^^^^^^^^^ AUTORTFM_PLATFORM_LINUX ^^^^^^^^^^^^^^^^^^^^^^^^^^^^^^
 #endif
 
-////////////////////////////////////////////////////////////////////////////////
-// Register all the open -> closed functions
-////////////////////////////////////////////////////////////////////////////////
-UE_AUTORTFM_REGISTER_OPEN_TO_CLOSED_FUNCTIONS(
+// Do not add new entries to this function as it will be removed once we're
+// confident that the AEM files entirely cover these mappings.
+void* LookupStdLibFunctionAndEnsureIfFound(void* Open)
+{
+#define UE_AUTORTFM_MAP_OPEN_TO_CLOSED(Open, Closed) std::make_pair(reinterpret_cast<void*>(Open), reinterpret_cast<void*>(Closed))
+#define UE_AUTORTFM_MAP_OPEN_TO_CLOSED_OVERLOADED(Signature, Open, Closed) UE_AUTORTFM_MAP_OPEN_TO_CLOSED(static_cast< std::add_pointer_t<Signature> >(Open), static_cast< std::add_pointer_t<Signature> >(Closed))
+#define UE_AUTORTFM_MAP_OPEN_TO_SELF(Open) UE_AUTORTFM_MAP_OPEN_TO_CLOSED(Open, Open)
+#define UE_AUTORTFM_MAP_OPEN_TO_SELF_OVERLOADED(Signature, Open) UE_AUTORTFM_MAP_OPEN_TO_CLOSED(static_cast< std::add_pointer_t<Signature> >(Open), static_cast< std::add_pointer_t<Signature> >(Open))
+
+	static std::unordered_map<void*, void*> Mappings {
+
 #if AUTORTFM_ASAN_ENABLED
 	UE_AUTORTFM_MAP_OPEN_TO_SELF(__asan_addr_is_in_fake_stack),
 	UE_AUTORTFM_MAP_OPEN_TO_SELF(__asan_get_current_fake_stack),
@@ -1051,8 +1059,6 @@ UE_AUTORTFM_REGISTER_OPEN_TO_CLOSED_FUNCTIONS(
 	UE_AUTORTFM_MAP_OPEN_TO_CLOSED(FlsSetValue, RTFM_FlsSetValue),
 	UE_AUTORTFM_MAP_OPEN_TO_SELF(_Query_perf_frequency),
 	UE_AUTORTFM_MAP_OPEN_TO_SELF(_Query_perf_counter),
-#else
-	UE_AUTORTFM_MAP_OPEN_TO_SELF(_ZNSt3__112__next_primeEm),
 #endif
 
 #if AUTORTFM_PLATFORM_WINDOWS && AUTORTFM_EXCEPTIONS_ENABLED
@@ -1327,7 +1333,18 @@ UE_AUTORTFM_REGISTER_OPEN_TO_CLOSED_FUNCTIONS(
 	UE_AUTORTFM_MAP_OPEN_TO_CLOSED(wcscpy, RTFM_wcscpy),
 	UE_AUTORTFM_MAP_OPEN_TO_CLOSED(wcsncpy, RTFM_wcsncpy),
 	UE_AUTORTFM_MAP_OPEN_TO_CLOSED(atexit, RTFM_atexit)
-);
+
+	};
+
+	void* Closed = Mappings[Open];
+	AUTORTFM_ENSURE_MSG(!Closed,
+		"Closed function for '%s' found in LookupStdLibFunctionAndEnsureIfFound(), "
+		"but should have been declared in the external mapping files",
+		GetFunctionDescription(Open).c_str()
+	);
+	return Closed;
+}
+
 
 #ifdef _MSC_VER
 #pragma warning(pop)
diff --git a/Engine/Source/Runtime/AutoRTFM/Private/StdLibSupport.h b/Engine/Source/Runtime/AutoRTFM/Private/StdLibSupport.h
new file mode 100644
--- /dev/null
+++ b/Engine/Source/Runtime/AutoRTFM/Private/StdLibSupport.h
@@ -0,0 +1,20 @@
+// Copyright Epic Games, Inc. All Rights Reserved.
+
+#pragma once
+
+#if (defined(__AUTORTFM) && __AUTORTFM)
+
+namespace AutoRTFM
+{
+
+// A temporary fallback path for looking up standard library closed variants
+// functions from their open function address. If the closed function is found
+// then the function ensures as this should have been caught by a declaration
+// in an AutoRTFM External Mapping file. 
+// Do not add new entries to this function as it will be removed once we're
+// confident that the AEM files entirely cover these mappings.
+void* LookupStdLibFunctionAndEnsureIfFound(void* Open);
+
+}  // AutoRTFM
+
+#endif // (defined(__AUTORTFM) && __AUTORTFM)
diff --git a/Engine/Source/Runtime/AutoRTFM/Public/AutoRTFM.h b/Engine/Source/Runtime/AutoRTFM/Public/AutoRTFM.h
--- a/Engine/Source/Runtime/AutoRTFM/Public/AutoRTFM.h
+++ b/Engine/Source/Runtime/AutoRTFM/Public/AutoRTFM.h
@@ -1321,73 +1321,6 @@ namespace ForTheRuntime
 		return static_cast<ETransactionResult>(autortfm_cascading_rollback_transaction());
 	}
 
-	// A template helper that holds the mapping from the open function pointer
-	// OpenFn to the closed function pointer ClosedFn, both of the function type
-	// FnType.
-	template<typename FnType, FnType* OpenFn, FnType* ClosedFn>
-	class TOpenToClosedMapping
-	{
-	public:
-		// Returns true if the open function pointer is not null.
-		inline constexpr bool IsValid() const { return OpenFn != nullptr; }
-
-		// Returns the mapping as a autortfm_open_to_closed_mapping structure.
-		inline constexpr autortfm_open_to_closed_mapping Get() const
-		{
-			return {reinterpret_cast<void*>(OpenFn), reinterpret_cast<void*>(ClosedFn)};
-		}
-	};
-
-	// A helper that on construction, registers all the open to closed function
-	// mappings expressed as the variadic TOpenToClosedMapping template
-	// arguments. On destruction the mappings are unregistered.
-	// Used by UE_AUTORTFM_REGISTER_OPEN_TO_CLOSED_FUNCTIONS()
-	template <typename ... OpenToClosedMappings>
-	class TAutoRegisterOpenToClosedFunctions
-	{
-	public:
-		UE_AUTORTFM_NOAUTORTFM
-		inline TAutoRegisterOpenToClosedFunctions()
-		{
-			autortfm_register_open_to_closed_functions(GetTable());
-		}
-
-		UE_AUTORTFM_NOAUTORTFM
-		inline ~TAutoRegisterOpenToClosedFunctions()
-		{
-			autortfm_unregister_open_to_closed_functions(GetTable());
-		}
-
-	private:
-		struct FStaticData
-		{
-			FStaticData()
-			{
-				// Some functions may dynamically resolve to nullptr, which we want to filter out
-				// as this will be treated as a null-terminator.
-				size_t Count = 0;
-				auto AddMappingIfValid = [&](const auto& Mapping)
-				{
-					if (Mapping.IsValid())
-					{
-						Mappings[Count++] = Mapping.Get();
-					}
-				};
-				(AddMappingIfValid(OpenToClosedMappings{}), ...);
-				Mappings[Count++] = {nullptr, nullptr};
-			}
-			autortfm_open_to_closed_mapping Mappings[sizeof...(OpenToClosedMappings) + 1]{};
-			autortfm_open_to_closed_table Table{Mappings, /* Prev */ nullptr, /* Next */ nullptr};
-		};
-
-		UE_AUTORTFM_NOAUTORTFM
-		static autortfm_open_to_closed_table* GetTable()
-		{
-			static FStaticData StaticData;
-			return &StaticData.Table;
-		}
-	};
-
 	// Reserved for future.
 	UE_AUTORTFM_CRITICAL_INLINE void RecordOpenRead(void const*, size_t) {}
 
@@ -1524,24 +1457,4 @@ namespace AutoRTFM::Private
 #define UE_AUTORTFM_REGISTER_OPEN_TO_CLOSED_FUNCTIONS(...)
 #endif
 
-// Used as an argument to UE_AUTORTFM_REGISTER_OPEN_TO_CLOSED_FUNCTIONS()
-// Maps the open function Open to the closed function Closed
-#define UE_AUTORTFM_MAP_OPEN_TO_CLOSED(Open, Closed) ::AutoRTFM::ForTheRuntime::TOpenToClosedMapping<decltype(Open), Open, Closed>
-
-// Used as an argument to UE_AUTORTFM_REGISTER_OPEN_TO_CLOSED_FUNCTIONS()
-// Similar to AUTORTFM_MAP_OPEN_TO_CLOSED(), but takes an additional Signature 
-// parameter used to distinguish the overload that is to be registered.
-// Signature is in the form: ReturnType(ArgType1, ArgType2, ...)
-#define UE_AUTORTFM_MAP_OPEN_TO_CLOSED_OVERLOADED(Signature, Open, Closed) ::AutoRTFM::ForTheRuntime::TOpenToClosedMapping<Signature, Open, Closed>
-
-// Used as an argument to UE_AUTORTFM_REGISTER_OPEN_TO_CLOSED_FUNCTIONS()
-// Maps the open function Open to itself as the closed function
-#define UE_AUTORTFM_MAP_OPEN_TO_SELF(Open) ::AutoRTFM::ForTheRuntime::TOpenToClosedMapping<decltype(Open), Open, Open>
-
-// Used as an argument to UE_AUTORTFM_REGISTER_OPEN_TO_CLOSED_FUNCTIONS()
-// Similar to AUTORTFM_MAP_OPEN_TO_SELF(), but takes an additional Signature 
-// parameter used to distinguish the overload that is to be registered.
-// Signature is in the form: ReturnType(ArgType1, ArgType2, ...)
-#define UE_AUTORTFM_MAP_OPEN_TO_SELF_OVERLOADED(Signature, Open) ::AutoRTFM::ForTheRuntime::TOpenToClosedMapping<Signature, Open, Open>
-
 #endif // __cplusplus
