[AutoRTFM] Allow external mappings for non-inlinable fuctions, add Internal.aem

If the function is marked non-inlinable then we can safely use an AEM to replace the closed variant.
Add Internal.aem to map some of the AutoRTFM runtime functions to their closed variant, removing some bespoke logic from the compiler.

diff --git a/Engine/Source/Runtime/AutoRTFM/AutoRTFM.Build.cs b/Engine/Source/Runtime/AutoRTFM/AutoRTFM.Build.cs
--- a/Engine/Source/Runtime/AutoRTFM/AutoRTFM.Build.cs
+++ b/Engine/Source/Runtime/AutoRTFM/AutoRTFM.Build.cs
@@ -16,6 +16,7 @@ namespace UnrealBuildTool.Rules
 			);
 			PrivateDefinitions.Add("SUPPRESS_PER_MODULE_INLINE_FILE"); // This module does not use core's standard operator new/delete overloads
 
+			AutoRTFMExternalMappingFiles.Add("Runtime/AutoRTFM/Public/Internal.aem");
 			AutoRTFMExternalMappingFiles.Add("Runtime/AutoRTFM/Public/StdLib.Common.aem");
 			
 			if (Target.Platform == UnrealTargetPlatform.Win64)
diff --git a/Engine/Source/Runtime/AutoRTFM/Private/API.cpp b/Engine/Source/Runtime/AutoRTFM/Private/API.cpp
--- a/Engine/Source/Runtime/AutoRTFM/Private/API.cpp
+++ b/Engine/Source/Runtime/AutoRTFM/Private/API.cpp
@@ -906,23 +906,15 @@ extern "C" UE_AUTORTFM_NOAUTORTFM bool autortfm_current_nest_throw() noexcept
 	return true;
 }
 
+// Note: Internal.aem maps this to autortfm_open_from_closed_explicit_validation when called in the closed
 extern "C" void autortfm_open_explicit_validation(autortfm_memory_validation_level ValidationLevel, void (*Work)(void*), void* Arg, const void* ReturnAddress)
 {
-	if (autortfm_is_closed())
-	{
-		return autortfm_open_from_closed_explicit_validation(ValidationLevel, Work, Arg, ReturnAddress);
-	}
-
 	Work(Arg);
 }
 
+// Note: Internal.aem maps this to autortfm_open_from_closed when called in the closed
 extern "C" void autortfm_open(void (*Work)(void*), void* Arg, const void* ReturnAddress)
 {
-	if (autortfm_is_closed())
-	{
-		return autortfm_open_from_closed(Work, Arg, ReturnAddress);
-	}
-
 	Work(Arg);
 }
 
diff --git a/Engine/Source/Runtime/AutoRTFM/Private/CompilerInterface.cpp b/Engine/Source/Runtime/AutoRTFM/Private/CompilerInterface.cpp
--- a/Engine/Source/Runtime/AutoRTFM/Private/CompilerInterface.cpp
+++ b/Engine/Source/Runtime/AutoRTFM/Private/CompilerInterface.cpp
@@ -87,7 +87,8 @@ extern "C" AUTORTFM_DISABLE UE_AUTORTFM_API AUTORTFM_NO_ASAN void autortfm_recor
 	}
 }
 
-extern "C" UE_AUTORTFM_ALWAYS_OPEN UE_AUTORTFM_API void* autortfm_lookup_function(void* OriginalFunction, const char* Where) noexcept
+// Note: Internal.aem maps this to itself when called in the closed
+extern "C" void* autortfm_lookup_function(void* OriginalFunction, const char* Where) noexcept
 {
 	return FunctionMapLookup(OriginalFunction, Where);
 }
diff --git a/Engine/Source/Runtime/AutoRTFM/Public/AutoRTFM.h b/Engine/Source/Runtime/AutoRTFM/Public/AutoRTFM.h
--- a/Engine/Source/Runtime/AutoRTFM/Public/AutoRTFM.h
+++ b/Engine/Source/Runtime/AutoRTFM/Public/AutoRTFM.h
@@ -259,13 +259,13 @@ UE_AUTORTFM_CRITICAL_INLINE autortfm_status autortfm_get_context_status() AUTORT
 #endif
 
 #if UE_AUTORTFM_ENABLED
-UE_AUTORTFM_API void autortfm_open(void (*work)(void* arg), void* arg, const void* return_address);
+UE_AUTORTFM_FORCENOINLINE UE_AUTORTFM_API void autortfm_open(void (*work)(void* arg), void* arg, const void* return_address);
 #else
 UE_AUTORTFM_CRITICAL_INLINE void autortfm_open(void (*work)(void* arg), void* arg, const void* /* return_address */) { work(arg); }
 #endif
 
 #if UE_AUTORTFM_ENABLED
-UE_AUTORTFM_API void autortfm_open_explicit_validation(autortfm_memory_validation_level, void (*work)(void* arg), void* arg, const void* return_address);
+UE_AUTORTFM_FORCENOINLINE UE_AUTORTFM_API void autortfm_open_explicit_validation(autortfm_memory_validation_level, void (*work)(void* arg), void* arg, const void* return_address);
 #else
 UE_AUTORTFM_CRITICAL_INLINE void autortfm_open_explicit_validation(autortfm_memory_validation_level, void (*work)(void* arg), void* arg, const void* /* return_address */) { work(arg); }
 #endif
@@ -546,7 +546,7 @@ template<typename FunctorType>
 void AutoRTFMFunctorInvoker(void* Arg) { (*static_cast<const FunctorType*>(Arg))(); }
 
 #if UE_AUTORTFM
-extern "C" UE_AUTORTFM_API void* autortfm_lookup_function(void* OriginalFunction, const char* Where) AUTORTFM_NOEXCEPT;
+extern "C" UE_AUTORTFM_FORCENOINLINE UE_AUTORTFM_API void* autortfm_lookup_function(void* OriginalFunction, const char* Where) AUTORTFM_NOEXCEPT;
 
 template<typename FunctorType>
 auto AutoRTFMLookupInstrumentedFunctorInvoker(const FunctorType& Functor) -> void(*)(void*)
