[AutoRTFM] Optimize the write log in both space and time.

Previously, the transaction would split all writes into 32K chunks, and the write log would waste partial blocks. Now, the transaction passes the write through to the write log as-is, and the write log splits the data across blocks as needed. This avoids wasted memory when large write log changes occur, such as via memcpy() or memset().

Small writes now go through a new `PushSmall` path, which ignores write splitting entirely and also uses a template to bake in the write size. This makes simple writes 15-28% faster.

Benchmark sheet: https://docs.google.com/spreadsheets/d/11ZTSQOOQ6SMHa3EUR_AQxerdciT9ZTv0loERWvB53zc/edit?usp=sharing

diff --git a/Engine/Source/Programs/AutoRTFMTests/Private/WriteLogTests.cpp b/Engine/Source/Programs/AutoRTFMTests/Private/WriteLogTests.cpp
--- a/Engine/Source/Programs/AutoRTFMTests/Private/WriteLogTests.cpp
+++ b/Engine/Source/Programs/AutoRTFMTests/Private/WriteLogTests.cpp
@@ -5,6 +5,7 @@
 #include <algorithm>
 #include <cstddef>
 #include <numeric>
+#include <random>
 #include <vector>
 
 TEST_CASE("FWriteLog")
@@ -37,7 +38,7 @@ TEST_CASE("FWriteLog")
 		std::vector<size_t> EntrySizes(NumEntries);
 		for (size_t& EntrySize : EntrySizes)
 		{
-			EntrySize = std::max<size_t>(1, Rand() & AutoRTFM::FWriteLogEntry::MaxSize);
+			EntrySize = std::uniform_int_distribution<unsigned int>(1, AutoRTFM::FWriteLog::RecordMaxSize)(Rand);
 		}
 		return EntrySizes;
 	};
@@ -108,50 +109,72 @@ TEST_CASE("FWriteLog")
 		}
 
 		REQUIRE(!WriteLog.IsEmpty());
-		REQUIRE(WriteLog.Num() <= NumEntries); // WriteLog may append back-to-back entries.
 		REQUIRE(WriteLog.TotalSize() == Data.size());
 
 		SECTION("Forwards iterator")
 		{
-			size_t EntryIndex = 0;
-			for (AutoRTFM::FWriteLogEntry Got : WriteLog)
+			AutoRTFM::FWriteLogEntry Expect;
+			ptrdiff_t EntryIndex = 0;
+			for (AutoRTFM::FWriteLog::Iterator It = WriteLog.begin(); It != WriteLog.end(); ++It)
 			{
-				size_t Offset = 0;
-				while (Offset < Got.Size)
+				AutoRTFM::FWriteLogEntry Got = *It;
+				while (Got.Size > 0)
 				{
-					const AutoRTFM::FWriteLogEntry& Expect = Entries[EntryIndex];
-					REQUIRE(Got.LogicalAddress + Offset == Expect.LogicalAddress);
+					// Move to the next expectation record when this one is fully empty.
+					if (Expect.Size == 0)
+					{
+						Expect = Entries[EntryIndex++];
+					}
+
+					// The actual write log must match the expectation's address and memory-validation.
+					REQUIRE(Got.LogicalAddress == Expect.LogicalAddress);
 					REQUIRE(Got.bNoMemoryValidation == Expect.bNoMemoryValidation);
-					REQUIRE(Expect.Size + Offset <= Got.Size);
-					REQUIRE(memcmp(Got.Data + Offset, Expect.Data, Expect.Size) == 0);
-					EntryIndex++;
-					Offset += Expect.Size;
+
+					// Compare the actual contents of the write log with our expectations.
+					const size_t NumBytes = std::min(Got.Size, Expect.Size);
+					REQUIRE(memcmp(Got.Data, Expect.Data, NumBytes) == 0);
+					Got.LogicalAddress += NumBytes;
+					Got.Data += NumBytes;
+					Got.Size -= NumBytes;
+					Expect.LogicalAddress += NumBytes;
+					Expect.Data += NumBytes;
+					Expect.Size -= NumBytes;
 				}
-				REQUIRE(Offset == Got.Size);
+				REQUIRE(Got.Size == 0);
 			}
 			REQUIRE(EntryIndex == Entries.size());
 		}
 
 		SECTION("Reverse iterator")
 		{
-			size_t EntryIndex = Entries.size() - 1;
+			AutoRTFM::FWriteLogEntry Expect;
+			ptrdiff_t EntryIndex = Entries.size();
 			for (AutoRTFM::FWriteLog::ReverseIterator It = WriteLog.rbegin(); It != WriteLog.rend(); ++It)
 			{
-				const AutoRTFM::FWriteLogEntry Got = *It;
-				size_t Offset = Got.Size;
-				while (Offset > 0)
+				AutoRTFM::FWriteLogEntry Got = *It;
+				while (Got.Size > 0)
 				{
-					const AutoRTFM::FWriteLogEntry& Expect = Entries[EntryIndex];
-					Offset -= Expect.Size;
-					REQUIRE(Offset < Got.Size);
-					REQUIRE(Got.LogicalAddress + Offset == Expect.LogicalAddress);
+					// Move to the previous expectation record when this one is fully empty.
+					if (Expect.Size == 0)
+					{
+						Expect = Entries[--EntryIndex];
+					}
+
+					// The actual write log must match the expectation's address and memory-validation.
+					REQUIRE(Got.LogicalAddress + Got.Size == Expect.LogicalAddress + Expect.Size);
 					REQUIRE(Got.bNoMemoryValidation == Expect.bNoMemoryValidation);
-					REQUIRE(Expect.Size + Offset <= Got.Size);
-					REQUIRE(memcmp(Got.Data + Offset, Expect.Data, Expect.Size) == 0);
-					EntryIndex--;
+
+					// Compare the actual contents of the write log with our expectations.
+					// Since we are reading the lists in reverse, we check the right edge of the data, and 
+					// don't need to increment pointers.
+					const size_t NumBytes = std::min(Got.Size, Expect.Size);
+					REQUIRE(memcmp(Got.Data + Got.Size - NumBytes, Expect.Data + Expect.Size - NumBytes, NumBytes) == 0);
+					Got.Size -= NumBytes;
+					Expect.Size -= NumBytes;
 				}
+				REQUIRE(Got.Size == 0);
 			}
-			REQUIRE(EntryIndex + 1 == 0);
+			REQUIRE(EntryIndex == 0);
 		}
 
 		SECTION("Reset")
@@ -186,6 +209,12 @@ TEST_CASE("FWriteLog")
 		REQUIRE(WriteLog.TotalSize() == 0u);
 	}
 
+
+	SECTION("EntrySize: RecordMaxSize-1, NumEntries: 32")
+	{
+		Check(FixedSize(/* EntrySize */ AutoRTFM::FWriteLog::RecordMaxSize - 1, /* NumEntries */ 32));
+	}
+
 	SECTION("EntrySize: 1, NumEntries: 32")
 	{
 		Check(FixedSize(/* EntrySize */ 1, /* NumEntries */ 32));
@@ -201,14 +230,19 @@ TEST_CASE("FWriteLog")
 		Check(FixedSize(/* EntrySize */ 1024, /* NumEntries */ 32));
 	}
 
-	SECTION("EntrySize: FWriteLogEntry::MaxSize-1, NumEntries: 32")
+	SECTION("EntrySize: RecordMaxSize, NumEntries: 32")
 	{
-		Check(FixedSize(/* EntrySize */ AutoRTFM::FWriteLogEntry::MaxSize-1, /* NumEntries */ 32));
+		Check(FixedSize(/* EntrySize */ AutoRTFM::FWriteLog::RecordMaxSize, /* NumEntries */ 32));
 	}
 
-	SECTION("EntrySize: FWriteLogEntry::MaxSize, NumEntries: 32")
+	SECTION("EntrySize: RecordMaxSize+1, NumEntries: 32")
 	{
-		Check(FixedSize(/* EntrySize */ AutoRTFM::FWriteLogEntry::MaxSize, /* NumEntries */ 32));
+		Check(FixedSize(/* EntrySize */ AutoRTFM::FWriteLog::RecordMaxSize+1, /* NumEntries */ 32));
+	}
+
+	SECTION("EntrySize: RecordMaxSize*10, NumEntries: 10")
+	{
+		Check(FixedSize(/* EntrySize */ AutoRTFM::FWriteLog::RecordMaxSize * 10, /* NumEntries */ 10));
 	}
 
 	SECTION("EntrySize: random, NumEntries: 32")
@@ -252,11 +286,13 @@ TEST_CASE("FWriteLog.Hash")
 		size_t Offset = AlignmentOffset;
 		while (Offset + WriteSize < Buffer.Size)
 		{
-			AutoRTFM::FWriteLogEntry Entry{};
-			Entry.LogicalAddress = Buffer.Data + Offset;
-			Entry.Data = Entry.LogicalAddress;
-			Entry.Size = WriteSize;
-			WriteLog.Push(Entry);
+			WriteLog.Push(AutoRTFM::FWriteLogEntry
+			{
+				.LogicalAddress = Buffer.Data + Offset,
+				.Data = Buffer.Data + Offset,
+				.Size = WriteSize,
+				.bNoMemoryValidation = false
+			});
 
 			Writes.push_back(FMemorySpan{Buffer.Data + Offset, WriteSize});
 
diff --git a/Engine/Source/Runtime/AutoRTFM/Private/TransactionInlines.h b/Engine/Source/Runtime/AutoRTFM/Private/TransactionInlines.h
--- a/Engine/Source/Runtime/AutoRTFM/Private/TransactionInlines.h
+++ b/Engine/Source/Runtime/AutoRTFM/Private/TransactionInlines.h
@@ -98,23 +98,13 @@ void FTransaction::RecordWrite(void* LogicalAddress, size_t Size, bool bNoMemory
 
 	Stats.Collect<EStatsKind::NewMemoryTrackerMiss>();
 
-	FWriteLogEntry WriteLogEntry{};
-	WriteLogEntry.LogicalAddress = static_cast<std::byte*>(LogicalAddress);
-	WriteLogEntry.Data = static_cast<std::byte*>(LogicalAddress);
-	WriteLogEntry.bNoMemoryValidation = bNoMemoryValidation;
-	if (AUTORTFM_UNLIKELY(Size > FWriteLogEntry::MaxSize))
+	WriteLog.Push(FWriteLogEntry
 	{
-		WriteLogEntry.Size = FWriteLogEntry::MaxSize;
-		while (Size > FWriteLogEntry::MaxSize)
-		{
-			WriteLog.Push(WriteLogEntry);
-			WriteLogEntry.LogicalAddress += FWriteLogEntry::MaxSize;
-			WriteLogEntry.Data += FWriteLogEntry::MaxSize;
-			Size -= FWriteLogEntry::MaxSize;
-		}
-	}
-	WriteLogEntry.Size = Size;
-	WriteLog.Push(WriteLogEntry);
+		.LogicalAddress = static_cast<std::byte*>(LogicalAddress),
+		.Data = static_cast<std::byte*>(LogicalAddress),
+		.Size = Size,
+		.bNoMemoryValidation = bNoMemoryValidation
+	});
 }
 
 template<unsigned SIZE> AUTORTFM_NO_ASAN UE_AUTORTFM_FORCEINLINE void FTransaction::RecordWrite(void* LogicalAddress)
@@ -170,12 +160,7 @@ template<unsigned SIZE> AUTORTFM_NO_ASAN UE_AUTORTFM_FORCENOINLINE void FTransac
 
 	Stats.Collect<EStatsKind::NewMemoryTrackerMiss>();
 
-	FWriteLogEntry WriteLogEntry{};
-	WriteLogEntry.LogicalAddress = static_cast<std::byte*>(LogicalAddress);
-	WriteLogEntry.Data = static_cast<std::byte*>(LogicalAddress);
-	WriteLogEntry.Size = SIZE;
-
-	WriteLog.Push(WriteLogEntry);
+	WriteLog.PushSmall<SIZE>(static_cast<std::byte*>(LogicalAddress));
 }
 
 UE_AUTORTFM_FORCEINLINE void FTransaction::DidAllocate(void* LogicalAddress, const size_t Size)
diff --git a/Engine/Source/Runtime/AutoRTFM/Private/WriteLog.h b/Engine/Source/Runtime/AutoRTFM/Private/WriteLog.h
--- a/Engine/Source/Runtime/AutoRTFM/Private/WriteLog.h
+++ b/Engine/Source/Runtime/AutoRTFM/Private/WriteLog.h
@@ -8,6 +8,7 @@
 #include "BuildMacros.h"
 #include "Utils.h"
 
+#include <algorithm>
 #include <cstddef>
 #include <cstdint>
 #include <cstring>
@@ -17,22 +18,13 @@ namespace AutoRTFM
 {
 	struct FWriteLogEntry final
 	{
-		// Number of bits used by the FWriteLog to represent a write's size.
-		static constexpr size_t SizeBits = 15;
-
-		// The maximum size for a single write log entry.
-		// Split into multiple entries if the write is too large.
-		static constexpr size_t MaxSize = (1u << SizeBits) - 1;
-
 		// The address of the write.
 		std::byte* LogicalAddress = nullptr;
 		
 		// A pointer to the original data before the write occurred.
 		std::byte* Data = nullptr;
 		
-		// The size of the write in bytes. Must be smaller than MaxSize.
-		// If the write exceeds MaxSize, then the write must be split into
-		// multiple entries.
+		// The size of the write in bytes.
 		size_t Size = 0;
 		
 		// If true, then this write will not be considered by the AutoRTFM
@@ -45,6 +37,15 @@ namespace AutoRTFM
 	// Ensure changes to this class are kept in sync with Unreal.natvis.
 	class FWriteLog final
 	{
+	public:
+		// Number of bits used by the FRecord to represent a write's size.
+		static constexpr size_t RecordSizeBits = 15;
+
+		// The maximum size for a single write log entry.
+		// Writes will be split into multiple entries if the write is too large.
+		static constexpr size_t RecordMaxSize = (1u << RecordSizeBits) - 1;
+
+	private:
 		struct FRecord
 		{
 			uintptr_t Address             : 48;
@@ -119,61 +120,91 @@ namespace AutoRTFM
 				return LastRecord > FirstRecord;
 			}
 
-			// The result enumerator of Push()
-			enum class EPushResult
+			// This will return 0 if the passed-in entry can be folded into the previous write, and 1 if a new
+			// write record will be required.
+			UE_AUTORTFM_FORCEINLINE int NumRecordsNeededForWrite(std::byte* LogicalAddress, bool bNoMemoryValidation)
 			{
-				// The block does not have enough capacity to fit the entry.
-				Full,
-				// The block added the entry as a new write.
-				Added,
-				// The block folded the result into the end of the last write.
-				Folded,
-			};
-
-			// Attempts to add the entry into this block by copying the entry's data and creating a
-			// new record.
-			// Returns true if the entry was added, or false if the block does not have the capacity
-			// for the entry.
-			UE_AUTORTFM_FORCEINLINE EPushResult Push(FWriteLogEntry Entry)
+				return static_cast<int>(
+					IsEmpty() ||
+					reinterpret_cast<uintptr_t>(LogicalAddress) != LastRecord->Address + LastRecord->Size ||
+					LastRecord->bNoMemoryValidation != bNoMemoryValidation);
+			}
+
+			// This determines the number of bytes that can safely be passed to `Push` below, given the number of
+			// bytes in the write entry (`EntrySize`) and whether or not those bytes can be folded into an existing
+			// record (which can be determined by calling `NumRecordsNeededForWrite` above). This algorithm does
+			// _not_ clamp the input size to RecordMaxSize. Instead, the block allocation logic in `Push` and
+			// `PushSmall` avoids creating an FBlock larger than MaxSize at all. (When very large pushes occur, the
+			// `Push` algorithm _will_ create a large FBlock to satisfy it, but this logic does not call
+			// `CalculatePushBytes`, and will create a new tail block afterwards.)
+			// 
+			// Returns zero if the block is entirely full, `EntrySize` if there's enough space, or a value
+			// somewhere in-between if the entry needs to be split.
+			UE_AUTORTFM_FORCEINLINE size_t CalculatePushBytes(unsigned int EntrySize, int NumRecordsNeeded)
 			{
-				EPushResult Result = EPushResult::Full;
+				std::byte* BlockEdge = reinterpret_cast<std::byte*>(LastRecord - NumRecordsNeeded);
+				constexpr ptrdiff_t MinimumWorthwhileSplitSize = 8;
 
-				if (!IsEmpty() &&
-					reinterpret_cast<uintptr_t>(Entry.LogicalAddress) == LastRecord->Address + LastRecord->Size &&
-					LastRecord->bNoMemoryValidation == Entry.bNoMemoryValidation)
+				if (DataEnd + EntrySize <= BlockEdge)
 				{
-					if (DataEnd + Entry.Size > reinterpret_cast<std::byte*>(LastRecord))
-					{
-						// Entry's data does not fit in the block's remaining space.
-						return EPushResult::Full;
-					}
-
-					LastRecord->Size += Entry.Size;
-					Result = EPushResult::Folded;
+					// We will fit the entire entry in the block with room to spare.
+#if AUTORTFM_BUILD_DEBUG
+					AUTORTFM_ASSERT(EntrySize <= RecordMaxSize);
+					AUTORTFM_ASSERT((NumRecordsNeeded == 1) || (LastRecord->Size + EntrySize <= ptrdiff_t(RecordMaxSize)));
+#endif
+					return EntrySize;
+				}
+				else if (ptrdiff_t BlockCapacity = BlockEdge - DataEnd; BlockCapacity >= MinimumWorthwhileSplitSize)
+				{
+					// We have enough data to fill up the entire block. This path will split the data across the end of this
+					// block and the start of the next block. We avoid this path when the block has fewer than eight bytes
+					// left, just as an efficiency measure, since it takes extra time to assemble two FRecords instead of one.
+#if AUTORTFM_BUILD_DEBUG
+					AUTORTFM_ASSERT(BlockCapacity <= ptrdiff_t(RecordMaxSize));
+					AUTORTFM_ASSERT((NumRecordsNeeded == 1) || (LastRecord->Size + BlockCapacity <= ptrdiff_t(RecordMaxSize)));
+#endif
+					return BlockCapacity;
 				}
 				else
 				{
-					if (DataEnd + Entry.Size > reinterpret_cast<std::byte*>(LastRecord - 1))
-					{
-						// Entry's data + new record does not fit in the block's remaining space.
-						return EPushResult::Full;
-					}
+					// The block is completely full.
+					return 0;
+				}
+			}
+
+			// Grows this block by copying `NumBytes` bytes from `DataIn`, representing data originally from `LogicalAddress`. 
+			// Asserts if the block does not have enough capacity to hold `NumBytes`. The caller should determine the available
+			// capacity ahead of time by calling `NumRecordsNeededForWrite` and `CalculatePushBytes`.
+			UE_AUTORTFM_FORCEINLINE void Push(std::byte* LogicalAddress, std::byte* DataIn, size_t NumBytes, bool bNoMemoryValidation, int NumRecordsNeeded)
+			{
+#if AUTORTFM_BUILD_DEBUG
+				std::byte* DataLimit = reinterpret_cast<std::byte*>(LastRecord - NumRecordsNeeded);
+				AUTORTFM_ASSERT(DataEnd + NumBytes <= DataLimit);
+				AUTORTFM_ASSERT(NumBytes <= RecordMaxSize);
+#endif
+
+				if (NumRecordsNeeded == 1)
+				{
 					LastRecord--;
-					LastRecord->Address = reinterpret_cast<uintptr_t>(Entry.LogicalAddress);
-					LastRecord->Size = Entry.Size;
-					LastRecord->bNoMemoryValidation = Entry.bNoMemoryValidation;
-					Result = EPushResult::Added;
+					LastRecord->Address = reinterpret_cast<uintptr_t>(LogicalAddress);
+					LastRecord->Size = NumBytes;
+					LastRecord->bNoMemoryValidation = bNoMemoryValidation;
+				}
+				else
+				{
+#if AUTORTFM_BUILD_DEBUG
+					AUTORTFM_ASSERT(NumRecordsNeeded == 0);
+#endif
+					LastRecord->Size += NumBytes;
 				}
 
-				memcpy(DataEnd, Entry.Data, Entry.Size);
+				memcpy(DataEnd, DataIn, NumBytes);
 
-				DataEnd += Entry.Size;
+				DataEnd += NumBytes;
 
 #if AUTORTFM_BUILD_DEBUG
 				AUTORTFM_ASSERT(DataEnd <= reinterpret_cast<std::byte*>(LastRecord));
 #endif
-
-				return Result;
 			}
 
 			// The next block in the linked list.
@@ -207,28 +238,111 @@ namespace AutoRTFM
 		// The log will make a copy of the FWriteLogEntry's data.
 		void Push(FWriteLogEntry Entry)
 		{
-			AUTORTFM_ASSERT(Entry.Size <= FWriteLogEntry::MaxSize);
+#if AUTORTFM_BUILD_DEBUG
 			AUTORTFM_ASSERT((reinterpret_cast<uintptr_t>(Entry.LogicalAddress) & 0xffff0000'00000000) == 0);
+#endif
+
+			{
+				TotalSizeBytes += Entry.Size;
+				int NumRecordsNeeded = TailBlock->NumRecordsNeededForWrite(Entry.LogicalAddress, Entry.bNoMemoryValidation);
+				size_t NumBytes = TailBlock->CalculatePushBytes(Entry.Size, NumRecordsNeeded);
+
+				if (NumBytes == Entry.Size)
+				{
+					// This push fits into our existing block.
+					TailBlock->Push(Entry.LogicalAddress, Entry.Data, NumBytes, Entry.bNoMemoryValidation, NumRecordsNeeded);
+					NumEntries += NumRecordsNeeded;
+					return;
+				}
+
+				// The push doesn't fit into the existing block...
+				if (NumBytes > 0)
+				{
+					// ... but we can still use up the remainder of the block.
+					TailBlock->Push(Entry.LogicalAddress, Entry.Data, NumBytes, Entry.bNoMemoryValidation, NumRecordsNeeded);
+					NumEntries += NumRecordsNeeded;
+
+					// Adjust the entry to point to the remaining unlogged bytes.
+					Entry.Size -= NumBytes;
+					Entry.LogicalAddress += NumBytes;
+					Entry.Data += NumBytes;
+				}
+			}
 
-			FBlock::EPushResult PushResult = TailBlock->Push(Entry);
-			if (PushResult == FBlock::EPushResult::Added)
+			// Calculate how many maxed-out RecordMaxSize entries we can make.
+			const size_t NumFullRecords = Entry.Size / RecordMaxSize;
+			// Calculate how many bytes will be remaining once we have emitted the full-size entries.
+			const size_t RemainingBytes = Entry.Size - (NumFullRecords * RecordMaxSize);
+			// Calculate the exact required size of the block.
+			const size_t RequiredSize =
+				sizeof(FBlock) +                                                   // FBlock header
+				(NumFullRecords * (RecordMaxSize + sizeof(FRecord))) +             // Bytes needed for full records
+				(LIKELY(RemainingBytes) ? (RemainingBytes + sizeof(FRecord)) : 0); // Bytes needed for trailing partial record, if any
+			// Add padding to the block to account for alignment.
+			const size_t AlignedSize = AlignUp(RequiredSize, sizeof(FRecord));
+
+			// Create a new empty tail block, large enough to hold the remainder of this push in its entirety, and
+			// never smaller than the default block size.
+			const size_t BlockSize = std::max(AlignedSize, FBlock::DefaultSize);
+			AllocateNewBlock(BlockSize);
+
+			// Push all of the full records.
+			for (size_t Index = NumFullRecords; Index--; )
 			{
-				NumEntries++;
+				constexpr int NumRecordsNeeded = 1;
+				constexpr size_t NumBytes = RecordMaxSize;
+
+				TailBlock->Push(Entry.LogicalAddress, Entry.Data, NumBytes, Entry.bNoMemoryValidation, NumRecordsNeeded);
+				NumEntries += NumRecordsNeeded;
+
+				Entry.Size -= NumBytes;
+				Entry.LogicalAddress += NumBytes;
+				Entry.Data += NumBytes;
 			}
-			else if (AUTORTFM_UNLIKELY(PushResult == FBlock::EPushResult::Full))
+
+			// Push the final, partial record.
+			if (LIKELY(RemainingBytes))
+			{
+				constexpr int NumRecordsNeeded = 1;
+
+				TailBlock->Push(Entry.LogicalAddress, Entry.Data, RemainingBytes, Entry.bNoMemoryValidation, NumRecordsNeeded);
+				NumEntries += NumRecordsNeeded;
+			}
+
+			// If we've just created an extra-large block, and alignment has caused it to be less than 100% full,
+			// we preemptively allocate a new default-sized block here. This avoids a scenario where a future
+			// PushSmall could accidentally overflow the record size. Normally this would be impossible because 
+			// a default-sized FBlock is smaller than RecordMaxSize, but in this case we are making an FBlock which 
+			// is much larger than normal.
+			if (AlignedSize > RecordMaxSize && AlignedSize > RequiredSize)
+			{
+				AllocateNewBlock(FBlock::DefaultSize);
+			}
+		}
+
+		// Adds the write log entry to the log; assumes a payload small enough that splitting is not beneficial.
+		// If you have large sizes, you should use `Push` instead so that splitting can occur.
+		template <unsigned int SIZE>
+		void PushSmall(std::byte* LogicalAddress)
+		{
+			static_assert(SIZE <= RecordMaxSize);  // This is a hard upper limit.
+#if AUTORTFM_BUILD_DEBUG
+			AUTORTFM_ASSERT((reinterpret_cast<uintptr_t>(LogicalAddress) & 0xffff0000'00000000) == 0);
+#endif
+
+			TotalSizeBytes += SIZE;
+
+			int NumRecordsNeeded = TailBlock->NumRecordsNeededForWrite(LogicalAddress, /*bNoMemoryValidation=*/false);
+			size_t NumBytes = TailBlock->CalculatePushBytes(SIZE, NumRecordsNeeded);
+
+			if (UNLIKELY(NumBytes != SIZE))
 			{
-				const size_t RequiredSize = AlignUp(sizeof(FBlock) + Entry.Size, alignof(FRecord)) + sizeof(FRecord);
-				FBlock* NewBlock = FBlock::Allocate(std::max(RequiredSize, FBlock::DefaultSize));
-				NewBlock->PrevBlock = TailBlock;
-				TailBlock->NextBlock = NewBlock;
-				TailBlock = NewBlock;
-
-				PushResult = NewBlock->Push(Entry);
-				AUTORTFM_ASSERT(PushResult == FBlock::EPushResult::Added);
-				NumEntries++;
+				AllocateNewBlock(FBlock::DefaultSize);
+				NumRecordsNeeded = 1;
 			}
 
-			TotalSizeBytes += Entry.Size;
+			TailBlock->Push(LogicalAddress, LogicalAddress, SIZE, /*bNoMemoryValidation=*/false, NumRecordsNeeded);
+			NumEntries += NumRecordsNeeded;
 		}
 
 		// Iterator for enumerating the writes of the log.
@@ -239,13 +353,12 @@ namespace AutoRTFM
 
 			TIterator(FBlock* StartBlock) : Block(StartBlock)
 			{
-				if constexpr (IS_FORWARD)
+				if (UNLIKELY(Block->IsEmpty()))
 				{
-					if (Block->IsEmpty())
+					if (UNLIKELY(!AdvanceBlock()))
 					{
-						// First block is fixed size and may be empty if the
-						// first write is larger than its fixed size.
-						Block = Block->NextBlock;
+						// The write log is entirely empty.
+						return;
 					}
 				}
 
@@ -256,12 +369,13 @@ namespace AutoRTFM
 			// Returns the entry at the current iterator's position.
 			FWriteLogEntry operator*() const
 			{
-				FWriteLogEntry Entry;
-				Entry.LogicalAddress = reinterpret_cast<std::byte*>(Record->Address);
-				Entry.Data = Data;
-				Entry.Size = Record->Size;
-				Entry.bNoMemoryValidation = Record->bNoMemoryValidation;
-				return Entry;
+				return FWriteLogEntry
+				{
+					.LogicalAddress = reinterpret_cast<std::byte*>(Record->Address),
+					.Data = Data,
+					.Size = Record->Size,
+					.bNoMemoryValidation = !!Record->bNoMemoryValidation,
+				};
 			}
 
 			// Progresses the iterator to the next entry
@@ -271,14 +385,11 @@ namespace AutoRTFM
 				{
 					if (Record == Block->LastRecord)
 					{
-						Block = Block->NextBlock;
-						if (!Block)
+						if (LIKELY(AdvanceBlock()))
 						{
-							Reset();
-							return;
+							Data = Block->DataStart();
+							Record = Block->FirstRecord;
 						}
-						Data = Block->DataStart();
-						Record = Block->FirstRecord;
 					}
 					else
 					{
@@ -290,14 +401,11 @@ namespace AutoRTFM
 				{
 					if (Record == Block->FirstRecord)
 					{
-						Block = Block->PrevBlock;
-						if (!Block || Block->IsEmpty())
+						if (LIKELY(AdvanceBlock()))
 						{
-							Reset();
-							return;
+							Data = Block->LastData();
+							Record = Block->LastRecord;
 						}
-						Data = Block->LastData();
-						Record = Block->LastRecord;
 					}
 					else
 					{
@@ -322,6 +430,26 @@ namespace AutoRTFM
 				Record = nullptr;
 			}
 
+			// Moves from this block to the next (if IS_FORWARD) or previous (if not IS_FORWARD),
+			// skipping any empty blocks. Returns true on success. If no more blocks exist, resets 
+			// the iterator and returns false.
+			UE_AUTORTFM_FORCEINLINE bool AdvanceBlock()
+			{
+				do
+				{
+					Block = IS_FORWARD ? Block->NextBlock : Block->PrevBlock;
+				} 
+				while (Block && Block->IsEmpty());
+
+				if (!Block)
+				{
+					Reset();
+					return false;
+				}
+
+				return true;
+			}
+
 			FBlock* Block = nullptr;
 			std::byte* Data = nullptr;
 			FRecord* Record = nullptr;
@@ -375,6 +503,14 @@ namespace AutoRTFM
 		FHash Hash(size_t NumWriteEntries) const;
 
 	private:
+		UE_AUTORTFM_FORCEINLINE void AllocateNewBlock(size_t Size)
+		{
+			FBlock* NewBlock = FBlock::Allocate(Size);
+			NewBlock->PrevBlock = TailBlock;
+			TailBlock->NextBlock = NewBlock;
+			TailBlock = NewBlock;
+		}
+
 		template<size_t SIZE>
 		static constexpr bool IsAlignedForTRecord = (SIZE & (alignof(FRecord) - 1)) == 0;
 
