[AutoRTFM] Optimize bitfield handling in FRecord.

Hand-written bit manipulation is more efficient than relying on Clang's bitfield support. Bespoke FRecord bit packing is about 14% more efficient for small writes, and 4-8% more efficient for large writes. Eight-byte writes are basically unchanged (?1.5%) but every other benchmark category gets a small win, consistently across multiple runs.

Benchmarks: https://docs.google.com/spreadsheets/d/1c5Xlwtn5dZQlaVBV4VMbPZcteeJKKzcEXYbDNviW7OI/edit?gid=15648869#gid=15648869

diff --git a/Engine/Source/Runtime/AutoRTFM/Private/WriteLog.h b/Engine/Source/Runtime/AutoRTFM/Private/WriteLog.h
--- a/Engine/Source/Runtime/AutoRTFM/Private/WriteLog.h
+++ b/Engine/Source/Runtime/AutoRTFM/Private/WriteLog.h
@@ -48,10 +48,45 @@ namespace AutoRTFM
 	private:
 		struct FRecord
 		{
-			uintptr_t Address             : 48;
-			uintptr_t bNoMemoryValidation : 1;
-			uintptr_t Size                : 15;
+			void Set(std::byte* Address, size_t Size, bool bNoMemoryValidation)
+			{
+#if AUTORTFM_BUILD_DEBUG
+				AUTORTFM_ASSERT((reinterpret_cast<uintptr_t>(Address) & 0xFFFF0000'00000000) == 0);
+				AUTORTFM_ASSERT(Size <= 0x7FFF);
+#endif
+				Bits = reinterpret_cast<uint64_t>(Address);
+				Bits <<= 1;
+				Bits |= bNoMemoryValidation;
+				Bits <<= 15;
+				Bits |= Size;
+			}
+
+			void Grow(size_t Amount)
+			{
+#if AUTORTFM_BUILD_DEBUG
+				AUTORTFM_ASSERT(Size() + Amount <= 0x7FFF);
+#endif
+				Bits += Amount;
+			}
+
+			std::byte* Address() const
+			{
+				return reinterpret_cast<std::byte*>(Bits >> 16);
+			}
+
+			bool NoMemoryValidation() const
+			{
+				return Bits & 0x8000;
+			}
+
+			size_t Size() const
+			{
+				return Bits & 0x7FFF;
+			}
+
+			uint64_t Bits;
 		};
+
 		static_assert(sizeof(uintptr_t) == 8, "assumption: a pointer is 8 bytes");
 		static_assert(sizeof(FRecord) == 8);
 
@@ -73,7 +108,7 @@ namespace AutoRTFM
 			// TotalSize is the total size of the allocated memory for the block including
 			// the FBlock header.
 			explicit FBlock(size_t TotalSize)
-			{ 
+			{
 				AUTORTFM_ENSURE((TotalSize & (alignof(FRecord) - 1)) == 0);
 				std::byte* End = reinterpret_cast<std::byte*>(this) + TotalSize;
 				DataEnd = DataStart();
@@ -99,7 +134,7 @@ namespace AutoRTFM
 			// Note: This block must have been allocated with a call to Allocate().
 			void Free()
 			{
-				delete [] reinterpret_cast<std::byte*>(this);
+				delete[] reinterpret_cast<std::byte*>(this);
 			}
 
 			// Returns a pointer to the data for the first entry
@@ -111,7 +146,7 @@ namespace AutoRTFM
 			// Returns a pointer to the data for the last entry
 			std::byte* LastData()
 			{
-				return DataEnd - LastRecord->Size;
+				return DataEnd - LastRecord->Size();
 			}
 
 			// Returns true if the block holds no entries.
@@ -126,8 +161,8 @@ namespace AutoRTFM
 			{
 				return static_cast<int>(
 					IsEmpty() ||
-					reinterpret_cast<uintptr_t>(LogicalAddress) != LastRecord->Address + LastRecord->Size ||
-					LastRecord->bNoMemoryValidation != bNoMemoryValidation);
+					LogicalAddress != LastRecord->Address() + LastRecord->Size() ||
+					LastRecord->NoMemoryValidation() != bNoMemoryValidation);
 			}
 
 			// This determines the number of bytes that can safely be passed to `Push` below, given the number of
@@ -150,7 +185,7 @@ namespace AutoRTFM
 					// We will fit the entire entry in the block with room to spare.
 #if AUTORTFM_BUILD_DEBUG
 					AUTORTFM_ASSERT(EntrySize <= RecordMaxSize);
-					AUTORTFM_ASSERT((NumRecordsNeeded == 1) || (LastRecord->Size + EntrySize <= ptrdiff_t(RecordMaxSize)));
+					AUTORTFM_ASSERT((NumRecordsNeeded == 1) || (LastRecord->Size() + EntrySize <= ptrdiff_t(RecordMaxSize)));
 #endif
 					return EntrySize;
 				}
@@ -161,7 +196,7 @@ namespace AutoRTFM
 					// left, just as an efficiency measure, since it takes extra time to assemble two FRecords instead of one.
 #if AUTORTFM_BUILD_DEBUG
 					AUTORTFM_ASSERT(BlockCapacity <= ptrdiff_t(RecordMaxSize));
-					AUTORTFM_ASSERT((NumRecordsNeeded == 1) || (LastRecord->Size + BlockCapacity <= ptrdiff_t(RecordMaxSize)));
+					AUTORTFM_ASSERT((NumRecordsNeeded == 1) || (LastRecord->Size() + BlockCapacity <= ptrdiff_t(RecordMaxSize)));
 #endif
 					return BlockCapacity;
 				}
@@ -186,16 +221,14 @@ namespace AutoRTFM
 				if (NumRecordsNeeded == 1)
 				{
 					LastRecord--;
-					LastRecord->Address = reinterpret_cast<uintptr_t>(LogicalAddress);
-					LastRecord->Size = NumBytes;
-					LastRecord->bNoMemoryValidation = bNoMemoryValidation;
+					LastRecord->Set(LogicalAddress, NumBytes, bNoMemoryValidation);
 				}
 				else
 				{
 #if AUTORTFM_BUILD_DEBUG
 					AUTORTFM_ASSERT(NumRecordsNeeded == 0);
 #endif
-					LastRecord->Size += NumBytes;
+					LastRecord->Grow(NumBytes);
 				}
 
 				memcpy(DataEnd, DataIn, NumBytes);
@@ -371,10 +404,10 @@ namespace AutoRTFM
 			{
 				return FWriteLogEntry
 				{
-					.LogicalAddress = reinterpret_cast<std::byte*>(Record->Address),
+					.LogicalAddress = reinterpret_cast<std::byte*>(Record->Address()),
 					.Data = Data,
-					.Size = Record->Size,
-					.bNoMemoryValidation = !!Record->bNoMemoryValidation,
+					.Size = Record->Size(),
+					.bNoMemoryValidation = Record->NoMemoryValidation(),
 				};
 			}
 
@@ -393,7 +426,7 @@ namespace AutoRTFM
 					}
 					else
 					{
-						Data += Record->Size;
+						Data += Record->Size();
 						Record--;
 					}
 				}
@@ -410,7 +443,7 @@ namespace AutoRTFM
 					else
 					{
 						Record++;
-						Data -= Record->Size;
+						Data -= Record->Size();
 					}
 				}
 			}
@@ -438,7 +471,7 @@ namespace AutoRTFM
 				do
 				{
 					Block = IS_FORWARD ? Block->NextBlock : Block->PrevBlock;
-				} 
+				}
 				while (Block && Block->IsEmpty());
 
 				if (!Block)
@@ -473,7 +506,7 @@ namespace AutoRTFM
 		void Reset()
 		{
 			// Skip HeadBlock, which is held as part of this structure.
-			FBlock* Block = HeadBlock->NextBlock; 
+			FBlock* Block = HeadBlock->NextBlock;
 			while (nullptr != Block)
 			{
 				FBlock* const Next = Block->NextBlock;
