[AutoRTFM] Add AUTORTFM_ASSERT_DEBUG macro.

This compiles away to nothing in Shipping and Test builds, but is enforced in other builds.

diff --git a/Engine/Source/Runtime/AutoRTFM/Private/Utils.h b/Engine/Source/Runtime/AutoRTFM/Private/Utils.h
--- a/Engine/Source/Runtime/AutoRTFM/Private/Utils.h
+++ b/Engine/Source/Runtime/AutoRTFM/Private/Utils.h
@@ -6,6 +6,7 @@
 #if (defined(__AUTORTFM) && __AUTORTFM)
 
 #include "AutoRTFM.h"
+#include "BuildMacros.h"
 #include "ExternAPI.h"
 
 #include <cstdarg>
@@ -267,6 +268,24 @@ inline constexpr T Lerp(T Start, T End, T Fraction)
 		AutoRTFM::DoAssert(Lambda); \
 	}
 
+// For places where zero cost is truly important, AUTORTFM_ASSERT_DEBUG will
+// disappear entirely in Shipping and Test builds. We still evaluate the expression,
+// on the off-chance that we are given an expression with side effects, and rely on
+// the compiler to optimize away a dead expression.
+#if AUTORTFM_BUILD_SHIPPING || AUTORTFM_BUILD_TEST
+#define AUTORTFM_ASSERT_DEBUG(exp) ((void)AUTORTFM_LIKELY(static_cast<bool>(exp)))
+#else
+#define AUTORTFM_ASSERT_DEBUG(exp) \
+	if (AUTORTFM_UNLIKELY(!(exp))) \
+	{ \
+		auto Lambda = [] \
+		{ \
+			AUTORTFM_FATAL("AUTORTFM_ASSERT_DEBUG(%s) failure", #exp); \
+		}; \
+		AutoRTFM::DoAssert(Lambda); \
+	}
+#endif
+
 // Same funkiness as AUTORTFM_ASSERT, except that it doesn't cause a fatal error.
 // TODO: Maybe upgrade the Info -> Warning?
 #define AUTORTFM_EXPECT(exp) \
diff --git a/Engine/Source/Runtime/AutoRTFM/Private/WriteLog.h b/Engine/Source/Runtime/AutoRTFM/Private/WriteLog.h
--- a/Engine/Source/Runtime/AutoRTFM/Private/WriteLog.h
+++ b/Engine/Source/Runtime/AutoRTFM/Private/WriteLog.h
@@ -50,10 +50,9 @@ namespace AutoRTFM
 		{
 			void Set(std::byte* Address, size_t Size, bool bNoMemoryValidation)
 			{
-#if AUTORTFM_BUILD_DEBUG
-				AUTORTFM_ASSERT((reinterpret_cast<uintptr_t>(Address) & 0xFFFF0000'00000000) == 0);
-				AUTORTFM_ASSERT(Size <= 0x7FFF);
-#endif
+				AUTORTFM_ASSERT_DEBUG((reinterpret_cast<uintptr_t>(Address) & 0xFFFF0000'00000000) == 0);
+				AUTORTFM_ASSERT_DEBUG(Size <= 0x7FFF);
+		
 				Bits = reinterpret_cast<uint64_t>(Address);
 				Bits <<= 1;
 				Bits |= bNoMemoryValidation;
@@ -63,9 +62,7 @@ namespace AutoRTFM
 
 			void Grow(size_t Amount)
 			{
-#if AUTORTFM_BUILD_DEBUG
-				AUTORTFM_ASSERT(Size() + Amount <= 0x7FFF);
-#endif
+				AUTORTFM_ASSERT_DEBUG(Size() + Amount <= 0x7FFF);
 				Bits += Amount;
 			}
 
@@ -183,10 +180,8 @@ namespace AutoRTFM
 				if (DataEnd + EntrySize <= BlockEdge)
 				{
 					// We will fit the entire entry in the block with room to spare.
-#if AUTORTFM_BUILD_DEBUG
-					AUTORTFM_ASSERT(EntrySize <= RecordMaxSize);
-					AUTORTFM_ASSERT((NumRecordsNeeded == 1) || (LastRecord->Size() + EntrySize <= ptrdiff_t(RecordMaxSize)));
-#endif
+					AUTORTFM_ASSERT_DEBUG(EntrySize <= RecordMaxSize);
+					AUTORTFM_ASSERT_DEBUG((NumRecordsNeeded == 1) || (LastRecord->Size() + EntrySize <= ptrdiff_t(RecordMaxSize)));
 					return EntrySize;
 				}
 				else if (ptrdiff_t BlockCapacity = BlockEdge - DataEnd; BlockCapacity >= MinimumWorthwhileSplitSize)
@@ -194,10 +189,8 @@ namespace AutoRTFM
 					// We have enough data to fill up the entire block. This path will split the data across the end of this
 					// block and the start of the next block. We avoid this path when the block has fewer than eight bytes
 					// left, just as an efficiency measure, since it takes extra time to assemble two FRecords instead of one.
-#if AUTORTFM_BUILD_DEBUG
-					AUTORTFM_ASSERT(BlockCapacity <= ptrdiff_t(RecordMaxSize));
-					AUTORTFM_ASSERT((NumRecordsNeeded == 1) || (LastRecord->Size() + BlockCapacity <= ptrdiff_t(RecordMaxSize)));
-#endif
+					AUTORTFM_ASSERT_DEBUG(BlockCapacity <= ptrdiff_t(RecordMaxSize));
+					AUTORTFM_ASSERT_DEBUG((NumRecordsNeeded == 1) || (LastRecord->Size() + BlockCapacity <= ptrdiff_t(RecordMaxSize)));
 					return BlockCapacity;
 				}
 				else
@@ -212,11 +205,8 @@ namespace AutoRTFM
 			// capacity ahead of time by calling `NumRecordsNeededForWrite` and `CalculatePushBytes`.
 			UE_AUTORTFM_FORCEINLINE void Push(std::byte* LogicalAddress, std::byte* DataIn, size_t NumBytes, bool bNoMemoryValidation, int NumRecordsNeeded)
 			{
-#if AUTORTFM_BUILD_DEBUG
-				std::byte* DataLimit = reinterpret_cast<std::byte*>(LastRecord - NumRecordsNeeded);
-				AUTORTFM_ASSERT(DataEnd + NumBytes <= DataLimit);
-				AUTORTFM_ASSERT(NumBytes <= RecordMaxSize);
-#endif
+				AUTORTFM_ASSERT_DEBUG(DataEnd + NumBytes <= reinterpret_cast<std::byte*>(LastRecord - NumRecordsNeeded));
+				AUTORTFM_ASSERT_DEBUG(NumBytes <= RecordMaxSize);
 
 				if (NumRecordsNeeded == 1)
 				{
@@ -225,9 +215,7 @@ namespace AutoRTFM
 				}
 				else
 				{
-#if AUTORTFM_BUILD_DEBUG
-					AUTORTFM_ASSERT(NumRecordsNeeded == 0);
-#endif
+					AUTORTFM_ASSERT_DEBUG(NumRecordsNeeded == 0);
 					LastRecord->Grow(NumBytes);
 				}
 
@@ -235,9 +223,7 @@ namespace AutoRTFM
 
 				DataEnd += NumBytes;
 
-#if AUTORTFM_BUILD_DEBUG
-				AUTORTFM_ASSERT(DataEnd <= reinterpret_cast<std::byte*>(LastRecord));
-#endif
+				AUTORTFM_ASSERT_DEBUG(DataEnd <= reinterpret_cast<std::byte*>(LastRecord));
 			}
 
 			// The next block in the linked list.
@@ -271,9 +257,7 @@ namespace AutoRTFM
 		// The log will make a copy of the FWriteLogEntry's data.
 		void Push(FWriteLogEntry Entry)
 		{
-#if AUTORTFM_BUILD_DEBUG
-			AUTORTFM_ASSERT((reinterpret_cast<uintptr_t>(Entry.LogicalAddress) & 0xffff0000'00000000) == 0);
-#endif
+			AUTORTFM_ASSERT_DEBUG((reinterpret_cast<uintptr_t>(Entry.LogicalAddress) & 0xffff0000'00000000) == 0);
 
 			{
 				TotalSizeBytes += Entry.Size;
@@ -359,9 +343,7 @@ namespace AutoRTFM
 		void PushSmall(std::byte* LogicalAddress)
 		{
 			static_assert(SIZE <= RecordMaxSize);  // This is a hard upper limit.
-#if AUTORTFM_BUILD_DEBUG
-			AUTORTFM_ASSERT((reinterpret_cast<uintptr_t>(LogicalAddress) & 0xffff0000'00000000) == 0);
-#endif
+			AUTORTFM_ASSERT_DEBUG((reinterpret_cast<uintptr_t>(LogicalAddress) & 0xffff0000'00000000) == 0);
 
 			TotalSizeBytes += SIZE;
 
