[AutoRTFM] Consider UnrealTargetPlatform.LinuxArm64 for StdLib.Linux.aem

Also fix what appears to be a harmless bad merge.

diff --git a/Engine/Source/Runtime/AutoRTFM/AutoRTFM.Build.cs b/Engine/Source/Runtime/AutoRTFM/AutoRTFM.Build.cs
--- a/Engine/Source/Runtime/AutoRTFM/AutoRTFM.Build.cs
+++ b/Engine/Source/Runtime/AutoRTFM/AutoRTFM.Build.cs
@@ -27,15 +27,7 @@ namespace UnrealBuildTool.Rules
 					AutoRTFMExternalMappingFiles.Add("Runtime/AutoRTFM/Public/ASAN.aem");
 				}
 			}
-			if (Target.Platform == UnrealTargetPlatform.Linux)
-			{
-				AutoRTFMExternalMappingFiles.Add("Runtime/AutoRTFM/Public/StdLib.Linux.aem");
-				if (Target.LinuxPlatform.bEnableAddressSanitizer)
-				{
-					AutoRTFMExternalMappingFiles.Add("Runtime/AutoRTFM/Public/ASAN.aem");
-				}
-			}
-			if (Target.Platform == UnrealTargetPlatform.Linux)
+			if (Target.Platform == UnrealTargetPlatform.Linux || Target.Platform == UnrealTargetPlatform.LinuxArm64)
 			{
 				AutoRTFMExternalMappingFiles.Add("Runtime/AutoRTFM/Public/StdLib.Linux.aem");
 				if (Target.LinuxPlatform.bEnableAddressSanitizer)
