[AutoRTFM] When allocating WriteLog blocks, increase the default allocation size over time.

Previously, a small push would always allocate a 2K block if the current block was full. Now, we start at 2K and increase the amount by 1.5x each time. We have a hard cap of 32K to avoid overflowing the FRecord maximum size.

This improves the Write benchmark dramatically for larger writes; consecutive int64 writes are 1.54x faster, and non-consecutive writes are 1.82x faster. In the extreme case, 100-byte non-consecutive writes are 4.45x faster.

Benchmark data: https://docs.google.com/spreadsheets/d/19C4Upv7Vv9RInrE_9LOm_QFhpR1u_MFOMdYRljNk6ME/edit?gid=15648869#gid=15648869

(Unfortunately, and a bit surprisingly, I wasn't able to detect any change in the VerseCLRAutoRTFM benchmark with this CL. Maybe it is mostly allocating new memory instead of overwriting existing data.)

diff --git a/Engine/Source/Runtime/AutoRTFM/Private/WriteLog.h b/Engine/Source/Runtime/AutoRTFM/Private/WriteLog.h
--- a/Engine/Source/Runtime/AutoRTFM/Private/WriteLog.h
+++ b/Engine/Source/Runtime/AutoRTFM/Private/WriteLog.h
@@ -98,7 +98,7 @@ namespace AutoRTFM
 			// Where:
 			//   Dₙ = Data n, Rₙ = Record n
 
-			// Size of a heap-allocated block, including the FBlock struct header.
+			// Starting size of a heap-allocated block, including the FBlock struct header.
 			static constexpr size_t DefaultSize = 2048;
 
 			// Constructor
@@ -121,7 +121,7 @@ namespace AutoRTFM
 			static FBlock* Allocate(size_t TotalSize)
 			{
 				AUTORTFM_ASSERT(TotalSize > (sizeof(FBlock) + sizeof(FRecord)));
-				std::byte* Memory = new std::byte[TotalSize];
+				void* Memory = AutoRTFM::Allocate(TotalSize, alignof(FBlock));
 				// Disable false-positive warning C6386: Buffer overrun while writing to 'Memory'
 				CA_SUPPRESS(6386)
 				return new (Memory) FBlock(TotalSize);
@@ -131,7 +131,7 @@ namespace AutoRTFM
 			// Note: This block must have been allocated with a call to Allocate().
 			void Free()
 			{
-				delete[] reinterpret_cast<std::byte*>(this);
+				AutoRTFM::Free(this);
 			}
 
 			// Returns a pointer to the data for the first entry
@@ -299,8 +299,8 @@ namespace AutoRTFM
 			const size_t AlignedSize = AlignUp(RequiredSize, sizeof(FRecord));
 
 			// Create a new empty tail block, large enough to hold the remainder of this push in its entirety, and
-			// never smaller than the default block size.
-			const size_t BlockSize = std::max(AlignedSize, FBlock::DefaultSize);
+			// never smaller than the upcoming block size.
+			const size_t BlockSize = std::max(AlignedSize, NextBlockSize);
 			AllocateNewBlock(BlockSize);
 
 			// Push all of the full records.
@@ -327,13 +327,13 @@ namespace AutoRTFM
 			}
 
 			// If we've just created an extra-large block, and alignment has caused it to be less than 100% full,
-			// we preemptively allocate a new default-sized block here. This avoids a scenario where a future
-			// PushSmall could accidentally overflow the record size. Normally this would be impossible because 
-			// a default-sized FBlock is smaller than RecordMaxSize, but in this case we are making an FBlock which 
-			// is much larger than normal.
-			if (AlignedSize > RecordMaxSize && AlignedSize > RequiredSize)
+			// we preemptively allocate a new block here. This avoids a scenario where a future PushSmall could
+			// accidentally overflow the record size. Normally this would be impossible because BlockMaxSize
+			// isn't large enough to hold more than RecordMaxSize bytes, but in this case we are making an FBlock
+			// which is potentially much larger than normal.
+			if (AlignedSize > BlockMaxSize && AlignedSize > RequiredSize)
 			{
-				AllocateNewBlock(FBlock::DefaultSize);
+				AllocateNewBlock(NextBlockSize);
 			}
 		}
 
@@ -352,7 +352,7 @@ namespace AutoRTFM
 
 			if (UNLIKELY(NumBytes != SIZE))
 			{
-				AllocateNewBlock(FBlock::DefaultSize);
+				AllocateNewBlock(NextBlockSize);
 				NumRecordsNeeded = 1;
 			}
 
@@ -495,11 +495,12 @@ namespace AutoRTFM
 				Block->Free();
 				Block = Next;
 			}
-			new (HeadBlockMemory) FBlock(HeadBlockSize - sizeof(FBlock));
+			new (HeadBlockMemory) FBlock(HeadBlockSize);
 			HeadBlock = reinterpret_cast<FBlock*>(HeadBlockMemory);
 			TailBlock = reinterpret_cast<FBlock*>(HeadBlockMemory);
 			NumEntries = 0;
 			TotalSizeBytes = 0;
+			NextBlockSize = FBlock::DefaultSize;
 		}
 
 		// Returns true if the log holds no entries.
@@ -524,15 +525,26 @@ namespace AutoRTFM
 			NewBlock->PrevBlock = TailBlock;
 			TailBlock->NextBlock = NewBlock;
 			TailBlock = NewBlock;
+
+			// Increase block sizes by 50% each time, capped at BlockMaxSize.
+			NextBlockSize = std::min<size_t>((NextBlockSize * 3 / 2), BlockMaxSize);
 		}
 
-		template<size_t SIZE>
-		static constexpr bool IsAlignedForTRecord = (SIZE & (alignof(FRecord) - 1)) == 0;
+		template <size_t SIZE>
+		static constexpr bool IsAlignedForFRecord = (SIZE & (alignof(FRecord) - 1)) == 0;
 
-		static constexpr size_t DefaultBlockSize = sizeof(FBlock) + FBlock::DefaultSize;
+		// The size of the inline block, which is declared as a byte array (`HeadBlockMemory`)
+		// inside the write log.
 		static constexpr size_t HeadBlockSize = 256;
-		static_assert(IsAlignedForTRecord<HeadBlockSize>);
-		static_assert(IsAlignedForTRecord<DefaultBlockSize>);
+
+		// The upper bound on heap-allocated block size. We avoid making blocks larger 
+		// than the maximum record size, so we don't need to insert size overflow 
+		// checks throughout the push logic.
+		static constexpr size_t BlockMaxSize = AlignDown(sizeof(FBlock) + sizeof(FRecord) + RecordMaxSize, alignof(FRecord));
+
+		static_assert(IsAlignedForFRecord<HeadBlockSize>);
+		static_assert(IsAlignedForFRecord<FBlock::DefaultSize>);
+		static_assert(IsAlignedForFRecord<BlockMaxSize>);
 
 		FHash HashAVX2(size_t NumWriteEntries) const;
 
@@ -540,6 +552,7 @@ namespace AutoRTFM
 		FBlock* TailBlock = reinterpret_cast<FBlock*>(HeadBlockMemory);
 		size_t NumEntries = 0;
 		size_t TotalSizeBytes = 0;
+		size_t NextBlockSize = FBlock::DefaultSize;
 		alignas(alignof(FBlock)) std::byte HeadBlockMemory[HeadBlockSize];
 	};
 }
