autortfm:
- add 'heap redirecting' feature
- for any pointers that we load from or store to that have an alternate (non-zero) address space, the instrumented code will redirect those loads and stores to runtime functions that can be supplied by the application

This is the compiler side changes required to support the DSTM Shared Heap feature: https://docs.google.com/document/d/1y6ZwIu7FGFXB8GOcBARMQRWnlqXo3yuKNcx_v_Lxl98

diff --git a/Engine/Source/Runtime/AutoRTFM/Private/API.cpp b/Engine/Source/Runtime/AutoRTFM/Private/API.cpp
--- a/Engine/Source/Runtime/AutoRTFM/Private/API.cpp
+++ b/Engine/Source/Runtime/AutoRTFM/Private/API.cpp
@@ -46,6 +46,10 @@ namespace
 	// ProcessAllPendingOpenToClosedRegistrations() when autortfm_initialize()
 	// is called.
 	autortfm_open_to_closed_table* GPendingOpenToClosedRegistrations = nullptr;
+
+	constexpr int GMaxRegisteredHeapRedirects = 4;
+	AutoRTFM::FHeapRedirectCallbacks GRegisteredHeapRedirects[GMaxRegisteredHeapRedirects];
+	int GNumRegisteredHeapRedirects = 0;
 }
 #endif  // UE_AUTORTFM
 
@@ -1150,6 +1154,48 @@ void ForTheRuntime::RegisterOnAbortFromTheOpen(TTask<void()> && Work)
 	RTFM_OnAbortInternal(std::move(Work));
 }
 
+void ForTheRuntime::RedirectedLoad(uint32 AddressSpace, void* DestPointer, uint64 Size, uint64 SourceAddress)
+{
+	for(int RedirectedHeapIndex = 0; RedirectedHeapIndex < GNumRegisteredHeapRedirects; RedirectedHeapIndex++)
+	{
+		FHeapRedirectCallbacks* HeapRedirects = &GRegisteredHeapRedirects[RedirectedHeapIndex];
+
+		if (HeapRedirects->AddressSpace == AddressSpace)
+		{
+			HeapRedirects->RedirectedLoad(DestPointer, Size, SourceAddress);
+			return;
+		}
+	}
+}
+
+void ForTheRuntime::RedirectedStore(uint32 AddressSpace, uint64 DestAddress, uint64 Size, const void* SourcePointer)
+{
+	for(int RedirectedHeapIndex = 0; RedirectedHeapIndex < GNumRegisteredHeapRedirects; RedirectedHeapIndex++)
+	{
+		FHeapRedirectCallbacks* HeapRedirects = &GRegisteredHeapRedirects[RedirectedHeapIndex];
+
+		if (HeapRedirects->AddressSpace == AddressSpace)
+		{
+			HeapRedirects->RedirectedStore(DestAddress, Size, SourcePointer);
+			return;
+		}
+	}
+}
+
+void RegisterHeapRedirectCallbacks(FHeapRedirectCallbacks Callbacks)
+{
+	AUTORTFM_ASSERT(GNumRegisteredHeapRedirects < GMaxRegisteredHeapRedirects);
+
+	// when loads and stores get redirected, it seems useful to not make a closed->open transition
+	// so that the redirect callbacks can abort the transaction, etc. So here in the registration
+	// function we translate the function pointers passed in to be their closed variants
+	FHeapRedirectCallbacks ClosedCallbacks;
+	ClosedCallbacks.AddressSpace = Callbacks.AddressSpace;
+	ClosedCallbacks.RedirectedLoad = reinterpret_cast<void (*)(void*, uint64_t, uint64_t)>(autortfm_lookup_function(reinterpret_cast<void*>(Callbacks.RedirectedLoad), "RegisterHeapRedirects"));
+	ClosedCallbacks.RedirectedStore = reinterpret_cast<void (*)(uint64_t, uint64_t, const void*)>(autortfm_lookup_function(reinterpret_cast<void*>(Callbacks.RedirectedStore), "RegisterHeapRedirects"));
+	GRegisteredHeapRedirects[GNumRegisteredHeapRedirects++] = ClosedCallbacks;
+}
+
 extern "C" void autortfm_on_commit(void (*Work)(void*), void* Arg)
 {
 	if (autortfm_is_closed())
diff --git a/Engine/Source/Runtime/AutoRTFM/Private/CompilerInterface.cpp b/Engine/Source/Runtime/AutoRTFM/Private/CompilerInterface.cpp
--- a/Engine/Source/Runtime/AutoRTFM/Private/CompilerInterface.cpp
+++ b/Engine/Source/Runtime/AutoRTFM/Private/CompilerInterface.cpp
@@ -87,6 +87,64 @@ extern "C" AUTORTFM_DISABLE UE_AUTORTFM_API AUTORTFM_NO_ASAN void autortfm_recor
 	}
 }
 
+extern "C" UE_AUTORTFM_API AUTORTFM_NO_ASAN void autortfm_redirected_load(uint32 AddressSpace, void* DestPointer, uint64 Size, uint64 SourceAddress) noexcept
+{
+	ForTheRuntime::RedirectedLoad(AddressSpace, DestPointer, Size, SourceAddress);
+}
+
+extern "C" UE_AUTORTFM_API AUTORTFM_NO_ASAN uint64 autortfm_redirected_load_8(uint32 AddressSpace, uint64 SourceAddress) noexcept
+{
+	uint64 Result = 0;
+	ForTheRuntime::RedirectedLoad(AddressSpace, &Result, 8, SourceAddress);
+	return Result;
+}
+
+extern "C" UE_AUTORTFM_API AUTORTFM_NO_ASAN uint32 autortfm_redirected_load_4(uint32 AddressSpace, uint64 SourceAddress) noexcept
+{
+	uint32 Result = 0;
+	ForTheRuntime::RedirectedLoad(AddressSpace, &Result, 4, SourceAddress);
+	return Result;
+}
+
+extern "C" UE_AUTORTFM_API AUTORTFM_NO_ASAN uint16 autortfm_redirected_load_2(uint32 AddressSpace, uint64 SourceAddress) noexcept
+{
+	uint16 Result = 0;
+	ForTheRuntime::RedirectedLoad(AddressSpace, &Result, 2, SourceAddress);
+	return Result;
+}
+
+extern "C" UE_AUTORTFM_API AUTORTFM_NO_ASAN uint8 autortfm_redirected_load_1(uint32 AddressSpace, uint64 SourceAddress) noexcept
+{
+	uint8 Result = 0;
+	ForTheRuntime::RedirectedLoad(AddressSpace, &Result, 1, SourceAddress);
+	return Result;
+}
+
+extern "C" UE_AUTORTFM_API AUTORTFM_NO_ASAN void autortfm_redirected_store(uint32 AddressSpace, uint64 DestAddress, uint64 Size, const void* SourcePointer) noexcept
+{
+	ForTheRuntime::RedirectedStore(AddressSpace, DestAddress, Size, SourcePointer);
+}
+
+extern "C" UE_AUTORTFM_API AUTORTFM_NO_ASAN void autortfm_redirected_store_8(uint32 AddressSpace, uint64 DestAddress, uint64 Value) noexcept
+{
+	ForTheRuntime::RedirectedStore(AddressSpace, DestAddress, sizeof(Value), &Value);
+}
+
+extern "C" UE_AUTORTFM_API AUTORTFM_NO_ASAN void autortfm_redirected_store_4(uint32 AddressSpace, uint64 DestAddress, uint32 Value) noexcept
+{
+	ForTheRuntime::RedirectedStore(AddressSpace, DestAddress, sizeof(Value), &Value);
+}
+
+extern "C" UE_AUTORTFM_API AUTORTFM_NO_ASAN void autortfm_redirected_store_2(uint32 AddressSpace, uint64 DestAddress, uint16 Value) noexcept
+{
+	ForTheRuntime::RedirectedStore(AddressSpace, DestAddress, sizeof(Value), &Value);
+}
+
+extern "C" UE_AUTORTFM_API AUTORTFM_NO_ASAN void autortfm_redirected_store_1(uint32 AddressSpace, uint64 DestAddress, uint8 Value) noexcept
+{
+	ForTheRuntime::RedirectedStore(AddressSpace, DestAddress, sizeof(Value), &Value);
+}
+
 // Note: Internal.aem maps this to itself when called in the closed
 extern "C" void* autortfm_lookup_function(void* OriginalFunction, const char* Where) noexcept
 {
diff --git a/Engine/Source/Runtime/AutoRTFM/Public/AutoRTFM.h b/Engine/Source/Runtime/AutoRTFM/Public/AutoRTFM.h
--- a/Engine/Source/Runtime/AutoRTFM/Public/AutoRTFM.h
+++ b/Engine/Source/Runtime/AutoRTFM/Public/AutoRTFM.h
@@ -538,6 +538,8 @@ namespace ForTheRuntime
 	UE_AUTORTFM_API void PopAllOnAbortHandlersInternal(const void* Key);
 	UE_AUTORTFM_API void RegisterOnCommitFromTheOpen(TTask<void()>&& Work);
 	UE_AUTORTFM_API void RegisterOnAbortFromTheOpen(TTask<void()>&& Work);
+	UE_AUTORTFM_API void RedirectedLoad(uint32_t AddressSpace, void* DestPointer, uint64_t Size, uint64_t SourceAddress);
+	UE_AUTORTFM_API void RedirectedStore(uint32_t AddressSpace, uint64_t DestAddress, uint64_t Size, const void* SourcePointer);
 } // namespace ForTheRuntime
 #endif
 
@@ -1026,6 +1028,21 @@ UE_AUTORTFM_CRITICAL_INLINE void PopAllOnAbortHandlers(const void* Key)
 }
 #endif
 
+struct FHeapRedirectCallbacks
+{
+	uint32_t AddressSpace;
+	void (*RedirectedLoad)(void* DestPointer, uint64_t Size, uint64_t SourceAddress);
+	void (*RedirectedStore)(uint64_t DestAddress, uint64_t Size, const void* SourcePointer);
+};
+
+#if UE_AUTORTFM
+void RegisterHeapRedirectCallbacks(FHeapRedirectCallbacks Callbacks);
+#else
+UE_AUTORTFM_CRITICAL_INLINE void RegisterHeapRedirectCallbacks(FHeapRedirectCallbacks /*Callbacks*/)
+{
+}
+#endif
+
 // Inform the runtime that we have performed a new object allocation. It's only
 // necessary to call this inside of custom malloc implementations. As an
 // optimization, you can choose to then only have your malloc return the pointer
