autortfm: suspected build fix?

diff --git a/Engine/Source/Runtime/AutoRTFM/Private/API.cpp b/Engine/Source/Runtime/AutoRTFM/Private/API.cpp
--- a/Engine/Source/Runtime/AutoRTFM/Private/API.cpp
+++ b/Engine/Source/Runtime/AutoRTFM/Private/API.cpp
@@ -1154,7 +1154,7 @@ void ForTheRuntime::RegisterOnAbortFromTheOpen(TTask<void()> && Work)
 	RTFM_OnAbortInternal(std::move(Work));
 }
 
-void ForTheRuntime::RedirectedLoad(uint32 AddressSpace, void* DestPointer, uint64 Size, uint64 SourceAddress)
+void ForTheRuntime::RedirectedLoad(uint32_t AddressSpace, void* DestPointer, uint64_t Size, uint64_t SourceAddress)
 {
 	for(int RedirectedHeapIndex = 0; RedirectedHeapIndex < GNumRegisteredHeapRedirects; RedirectedHeapIndex++)
 	{
@@ -1168,7 +1168,7 @@ void ForTheRuntime::RedirectedLoad(uint32 AddressSpace, void* DestPointer, uint6
 	}
 }
 
-void ForTheRuntime::RedirectedStore(uint32 AddressSpace, uint64 DestAddress, uint64 Size, const void* SourcePointer)
+void ForTheRuntime::RedirectedStore(uint32_t AddressSpace, uint64_t DestAddress, uint64_t Size, const void* SourcePointer)
 {
 	for(int RedirectedHeapIndex = 0; RedirectedHeapIndex < GNumRegisteredHeapRedirects; RedirectedHeapIndex++)
 	{
