[AutoRTFM] Fix `unsigned int` <-> `size_t` mismatch that's treated as an error with LLVM 20

diff --git a/Engine/Source/Runtime/AutoRTFM/Private/WriteLog.h b/Engine/Source/Runtime/AutoRTFM/Private/WriteLog.h
--- a/Engine/Source/Runtime/AutoRTFM/Private/WriteLog.h
+++ b/Engine/Source/Runtime/AutoRTFM/Private/WriteLog.h
@@ -172,7 +172,7 @@ namespace AutoRTFM
 			// 
 			// Returns zero if the block is entirely full, `EntrySize` if there's enough space, or a value
 			// somewhere in-between if the entry needs to be split.
-			UE_AUTORTFM_FORCEINLINE size_t CalculatePushBytes(unsigned int EntrySize, int NumRecordsNeeded)
+			UE_AUTORTFM_FORCEINLINE size_t CalculatePushBytes(size_t EntrySize, int NumRecordsNeeded)
 			{
 				std::byte* BlockEdge = reinterpret_cast<std::byte*>(LastRecord - NumRecordsNeeded);
 				constexpr ptrdiff_t MinimumWorthwhileSplitSize = 8;
