[AutoRTFM] [Linux] change prefix data to store the offset between the open and closed function so that it doesn't block relocation of those symbols.  This allows -MergeModules to link when verse-clang is used

- also update build.go to explicitly disable DIA SDK use since it gets silently switched on and causes problems when cmake finds it

diff --git a/Engine/Source/Programs/UnrealBuildTool/Platform/Linux/LinuxToolChain.cs b/Engine/Source/Programs/UnrealBuildTool/Platform/Linux/LinuxToolChain.cs
--- a/Engine/Source/Programs/UnrealBuildTool/Platform/Linux/LinuxToolChain.cs
+++ b/Engine/Source/Programs/UnrealBuildTool/Platform/Linux/LinuxToolChain.cs
@@ -499,14 +499,6 @@ namespace UnrealBuildTool
 				// Use local-dynamic TLS model. This generates less efficient runtime code for __thread variables, but avoids problems of running into
 				// glibc/ld.so limit (DTV_SURPLUS) for number of dlopen()'ed DSOs with static TLS (see e.g. https://www.cygwin.com/ml/libc-help/2013-11/msg00033.html)
 				Arguments.Add("-ftls-model=local-dynamic");
-
-				if (CompileEnvironment.bEnableAutoRTFMInstrumentation)
-				{
-					// Enabling the AutoRTFM "magic prefix" for a shared object will result in a
-					// R_X86_64_64 relocation linker error, due to the closed function address
-					// prefixed onto the open functions.
-					Arguments.Add("-Xclang -mllvm -Xclang -autortfm-no-magic-prefix");
-				}
 			}
 			else
 			{
diff --git a/Engine/Source/Runtime/AutoRTFM/Private/FunctionMapInlines.h b/Engine/Source/Runtime/AutoRTFM/Private/FunctionMapInlines.h
--- a/Engine/Source/Runtime/AutoRTFM/Private/FunctionMapInlines.h
+++ b/Engine/Source/Runtime/AutoRTFM/Private/FunctionMapInlines.h
@@ -22,7 +22,38 @@ void* FunctionMapLookupUsingMagicPrefix(void* OpenFn)
 	// to give us a much higher confidence that there is actually a closed
 	// variant pointer residing 8-bytes before our function address.
 	uint64_t PrefixData;
-	memcpy(&PrefixData, reinterpret_cast<char*>(OpenFn) - sizeof(uint64_t), sizeof(uint64_t));
+	char* CharOpenFn = reinterpret_cast<char*>(OpenFn);
+	memcpy(&PrefixData, CharOpenFn - sizeof(uint64_t), sizeof(uint64_t));
+#if PLATFORM_LINUX
+	constexpr uint64_t Mask = 0xffff'0000'0000'0000;
+
+	if (AUTORTFM_LIKELY(Constants::PosOffsetMagicPrefix == (PrefixData & Mask)))
+	{
+		return reinterpret_cast<void*>(CharOpenFn + (PrefixData & ~Mask));
+	}
+	else if (AUTORTFM_LIKELY((Constants::NegOffsetMagicPrefix) == (PrefixData & Mask)))
+	{
+		// when NegOffsetMagicPrefix is matched, that means the offset was negative and we need to sign extend
+		return reinterpret_cast<void*>(CharOpenFn + (int64_t)(PrefixData | Mask));
+	}
+	// UBSAN adds a type hash prefix to the function as a "prologue" that
+	// ends preceding our Magic Prefix. They use 0xc105cafe in the
+	// lower 32-bits to distinguish the 64-bit word containing their type
+	// hash. If we see it, check the preceding 64-bit word for our prefix.
+	else if (AUTORTFM_UNLIKELY(0xc105'cafe == (PrefixData & 0x0000'0000'ffff'ffff)))
+	{
+		memcpy(&PrefixData, reinterpret_cast<char*>(OpenFn) - sizeof(uint64_t) * 2, sizeof(uint64_t));
+		if (AUTORTFM_LIKELY(Constants::PosOffsetMagicPrefix == (PrefixData & Mask)))
+		{
+			return reinterpret_cast<void*>(CharOpenFn + (PrefixData & ~Mask));
+		}
+		else if (AUTORTFM_LIKELY((Constants::NegOffsetMagicPrefix) == (PrefixData & Mask)))
+		{
+			// when NegOffsetMagicPrefix is matched, that means the offset was negative and we need to sign extend
+			return reinterpret_cast<void*>(CharOpenFn + (int64_t)(PrefixData | Mask));
+		}
+	}
+#else
 	if (AUTORTFM_LIKELY(Constants::MagicPrefix == (PrefixData & 0xffff'0000'0000'0000)))
 	{
 		return reinterpret_cast<void*>(PrefixData & 0x0000'ffff'ffff'ffff);
@@ -39,6 +70,7 @@ void* FunctionMapLookupUsingMagicPrefix(void* OpenFn)
 			return reinterpret_cast<void*>(PrefixData & 0x0000'ffff'ffff'ffff);
 		}
 	}
+#endif
 
 	return nullptr;
 }
diff --git a/Engine/Source/Runtime/AutoRTFM/Private/Utils.cpp b/Engine/Source/Runtime/AutoRTFM/Private/Utils.cpp
--- a/Engine/Source/Runtime/AutoRTFM/Private/Utils.cpp
+++ b/Engine/Source/Runtime/AutoRTFM/Private/Utils.cpp
@@ -20,6 +20,12 @@
 namespace AutoRTFM
 {
 
+UE_AUTORTFM_FORCENOINLINE void DoAssert(void(*Logger)()) //-V1082 Silence PVS [[noreturn]] false positive warning
+{
+	Logger();
+	__builtin_unreachable();
+}
+
 std::string GetFunctionDescription(void* FunctionPtr)
 {
 #if AUTORTFM_PLATFORM_WINDOWS
diff --git a/Engine/Source/Runtime/AutoRTFM/Private/Utils.h b/Engine/Source/Runtime/AutoRTFM/Private/Utils.h
--- a/Engine/Source/Runtime/AutoRTFM/Private/Utils.h
+++ b/Engine/Source/Runtime/AutoRTFM/Private/Utils.h
@@ -53,11 +53,7 @@ std::string GetFunctionDescription(TReturnType (*FunctionPtr)(TParameterTypes...
 // fatal process destroying explosions happen here).
 [[clang::preserve_most]]
 [[noreturn]]
-UE_AUTORTFM_FORCENOINLINE inline void DoAssert(void(*Logger)()) //-V1082 Silence PVS [[noreturn]] false positive warning
-{
-	Logger();
-	__builtin_unreachable();
-}
+UE_AUTORTFM_FORCENOINLINE void DoAssert(void(*Logger)()); //-V1082 Silence PVS [[noreturn]] false positive warning
 
 // We use a special calling convention for this call to minimize
 // the cost of calling this function (which is super unlikely to
diff --git a/Engine/Source/Runtime/AutoRTFM/Public/AutoRTFMConstants.h b/Engine/Source/Runtime/AutoRTFM/Public/AutoRTFMConstants.h
--- a/Engine/Source/Runtime/AutoRTFM/Public/AutoRTFMConstants.h
+++ b/Engine/Source/Runtime/AutoRTFM/Public/AutoRTFMConstants.h
@@ -48,6 +48,8 @@ namespace AutoRTFM::Constants
 	// of an open function pointer address, we assume that we can find a closed
 	// variant pointer residing 8 bytes before the function address.
 	inline constexpr uint64_t MagicPrefix = 0xa273'0000'0000'0000;
+	inline constexpr uint64_t PosOffsetMagicPrefix = 0xa272'0000'0000'0000;
+	inline constexpr uint64_t NegOffsetMagicPrefix = 0xa271'0000'0000'0000;
 
 } // namespace AutoRTFM::Constants
 #endif
