[AutoRTFM] Use relative open->closed address for the open's prefix data on Linux

Absolute addresses can cause the linker to fail with R_X86_64_64 relocation linker errors.
Using relative addresses in the prefix data allows merge modules to link with AutoRTFM instrumentation enabled.
The UBT change to enable this will be done in a separate change.

Original change by james.singer, with an additional linker fix to ensure that the closed functions are placed in the same section as the open function.
This is done because there's a requirement in `ELFObjectWriter::recordRelocation()` that both SymbolA and SymbolB of MCValue (`(SymbolA - SymbolB + imm64)`) need to be in the same section.

diff --git a/Engine/Source/Runtime/AutoRTFM/Private/FunctionMapInlines.h b/Engine/Source/Runtime/AutoRTFM/Private/FunctionMapInlines.h
--- a/Engine/Source/Runtime/AutoRTFM/Private/FunctionMapInlines.h
+++ b/Engine/Source/Runtime/AutoRTFM/Private/FunctionMapInlines.h
@@ -15,6 +15,8 @@ namespace AutoRTFM
 AUTORTFM_DISABLE UE_AUTORTFM_FORCEINLINE
 void* FunctionMapLookupUsingMagicPrefix(void* OpenFn)
 {
+	constexpr uint64_t UbsanMagic = 0xc105'cafe;
+
 	// We use prefix data in our custom LLVM pass to stuff some data just
 	// before the address of all open function pointers (that we have
 	// definitions for!). We verify the special Magic Prefix constant in the
@@ -22,6 +24,37 @@ void* FunctionMapLookupUsingMagicPrefix(void* OpenFn)
 	// to give us a much higher confidence that there is actually a closed
 	// variant pointer residing 8-bytes before our function address.
 	uint64_t PrefixData;
+	char* CharOpenFn = reinterpret_cast<char*>(OpenFn);
+	memcpy(&PrefixData, CharOpenFn - sizeof(uint64_t), sizeof(uint64_t));
+#if PLATFORM_LINUX
+	constexpr uint64_t Mask = 0xffff'0000'0000'0000;
+	if (AUTORTFM_LIKELY(Constants::PosOffsetMagicPrefix == (PrefixData & Mask)))
+	{
+		return reinterpret_cast<void*>(CharOpenFn + (PrefixData & ~Mask));
+	}
+	else if (AUTORTFM_LIKELY((Constants::NegOffsetMagicPrefix) == (PrefixData & Mask)))
+	{
+		// when NegOffsetMagicPrefix is matched, that means the offset was negative and we need to sign extend
+		return reinterpret_cast<void*>(CharOpenFn + (int64_t)(PrefixData | Mask));
+	}
+	// UBSAN adds a type hash prefix to the function as a "prologue" that
+	// ends preceding our Magic Prefix. They use 0xc105cafe in the
+	// lower 32-bits to distinguish the 64-bit word containing their type
+	// hash. If we see it, check the preceding 64-bit word for our prefix.
+	else if (AUTORTFM_UNLIKELY(UbsanMagic == (PrefixData & 0x0000'0000'ffff'ffff)))
+	{
+		memcpy(&PrefixData, reinterpret_cast<char*>(OpenFn) - sizeof(uint64_t) * 2, sizeof(uint64_t));
+		if (AUTORTFM_LIKELY(Constants::PosOffsetMagicPrefix == (PrefixData & Mask)))
+		{
+			return reinterpret_cast<void*>(CharOpenFn + (PrefixData & ~Mask));
+		}
+		else if (AUTORTFM_LIKELY((Constants::NegOffsetMagicPrefix) == (PrefixData & Mask)))
+		{
+			// when NegOffsetMagicPrefix is matched, that means the offset was negative and we need to sign extend
+			return reinterpret_cast<void*>(CharOpenFn + (int64_t)(PrefixData | Mask));
+		}
+	}
+#else
 	memcpy(&PrefixData, reinterpret_cast<char*>(OpenFn) - sizeof(uint64_t), sizeof(uint64_t));
 	if (AUTORTFM_LIKELY(Constants::MagicPrefix == (PrefixData & 0xffff'0000'0000'0000)))
 	{
@@ -31,7 +64,7 @@ void* FunctionMapLookupUsingMagicPrefix(void* OpenFn)
 	// ends preceding our Magic Prefix. They use 0xc105cafe in the
 	// lower 32-bits to distinguish the 64-bit word containing their type
 	// hash. If we see it, check the preceding 64-bit word for our prefix.
-	else if (AUTORTFM_UNLIKELY(0xc105'cafe == (PrefixData & 0x0000'0000'ffff'ffff)))
+	else if (AUTORTFM_UNLIKELY(UbsanMagic == (PrefixData & 0x0000'0000'ffff'ffff)))
 	{
 		memcpy(&PrefixData, reinterpret_cast<char*>(OpenFn) - sizeof(uint64_t) * 2, sizeof(uint64_t));
 		if (AUTORTFM_LIKELY(Constants::MagicPrefix == (PrefixData & 0xffff'0000'0000'0000)))
@@ -39,6 +72,7 @@ void* FunctionMapLookupUsingMagicPrefix(void* OpenFn)
 			return reinterpret_cast<void*>(PrefixData & 0x0000'ffff'ffff'ffff);
 		}
 	}
+#endif
 
 	return nullptr;
 }
diff --git a/Engine/Source/Runtime/AutoRTFM/Private/Utils.cpp b/Engine/Source/Runtime/AutoRTFM/Private/Utils.cpp
--- a/Engine/Source/Runtime/AutoRTFM/Private/Utils.cpp
+++ b/Engine/Source/Runtime/AutoRTFM/Private/Utils.cpp
@@ -20,6 +20,17 @@
 namespace AutoRTFM
 {
 
+void DoAssert(void(*Logger)()) //-V1082 Silence PVS [[noreturn]] false positive warning
+{
+	Logger();
+	__builtin_unreachable();
+}
+
+void DoExpect(void(*Logger)())
+{
+	Logger();
+}
+
 std::string GetFunctionDescription(void* FunctionPtr)
 {
 #if AUTORTFM_PLATFORM_WINDOWS
diff --git a/Engine/Source/Runtime/AutoRTFM/Private/Utils.h b/Engine/Source/Runtime/AutoRTFM/Private/Utils.h
--- a/Engine/Source/Runtime/AutoRTFM/Private/Utils.h
+++ b/Engine/Source/Runtime/AutoRTFM/Private/Utils.h
@@ -51,23 +51,15 @@ std::string GetFunctionDescription(TReturnType (*FunctionPtr)(TParameterTypes...
 // the cost of calling this function (which is super unlikely to
 // be called because it is the asserting true codepath, EG. only
 // fatal process destroying explosions happen here).
-[[clang::preserve_most]]
-[[noreturn]]
-UE_AUTORTFM_FORCENOINLINE inline void DoAssert(void(*Logger)()) //-V1082 Silence PVS [[noreturn]] false positive warning
-{
-	Logger();
-	__builtin_unreachable();
-}
+[[clang::preserve_most, noreturn]]
+UE_AUTORTFM_FORCENOINLINE void DoAssert(void(*Logger)());
 
 // We use a special calling convention for this call to minimize
 // the cost of calling this function (which is super unlikely to
 // be called because it is the expecting true codepath, EG. only
 // things we don't expect to happen get here).
 [[clang::preserve_most]]
-UE_AUTORTFM_FORCENOINLINE inline void DoExpect(void(*Logger)())
-{
-	Logger();
-}
+UE_AUTORTFM_FORCENOINLINE void DoExpect(void(*Logger)());
 
 // Logs a message using a printf-style format string and va_list arguments.
 inline void LogV(const char* File, int Line, void* ProgramCounter, autortfm_log_severity Severity, const char* Format, va_list Args)
diff --git a/Engine/Source/Runtime/AutoRTFM/Public/AutoRTFMConstants.h b/Engine/Source/Runtime/AutoRTFM/Public/AutoRTFMConstants.h
--- a/Engine/Source/Runtime/AutoRTFM/Public/AutoRTFMConstants.h
+++ b/Engine/Source/Runtime/AutoRTFM/Public/AutoRTFMConstants.h
@@ -48,6 +48,11 @@ namespace AutoRTFM::Constants
 	// of an open function pointer address, we assume that we can find a closed
 	// variant pointer residing 8 bytes before the function address.
 	inline constexpr uint64_t MagicPrefix = 0xa273'0000'0000'0000;
+	// Similar to the above, but these constants indicate that the low 46 bits
+	// provide a relative addresses from the open function to the closed
+	// function.
+	inline constexpr uint64_t PosOffsetMagicPrefix = 0xa272'0000'0000'0000;
+	inline constexpr uint64_t NegOffsetMagicPrefix = 0xa271'0000'0000'0000;
 
 } // namespace AutoRTFM::Constants
 #endif
