[AutoRTFM] Fix bad code example in comment.

diff --git a/Engine/Source/Runtime/AutoRTFM/Public/AutoRTFM/OpenWrapper.h b/Engine/Source/Runtime/AutoRTFM/Public/AutoRTFM/OpenWrapper.h
--- a/Engine/Source/Runtime/AutoRTFM/Public/AutoRTFM/OpenWrapper.h
+++ b/Engine/Source/Runtime/AutoRTFM/Public/AutoRTFM/OpenWrapper.h
@@ -19,7 +19,7 @@ namespace AutoRTFM
 // failure or commit. This may be preferable to capturing the unwrapped shared
 // pointer as we can avoid costly closed transactional logic to ensure the
 // count is restored on failure:
-//    AutoRTFM::OnCommit([](WrappedSharedPtr = AutoRTFM::TOpenWrapper{MySharedPtr})
+//    AutoRTFM::OnCommit([WrappedSharedPtr = AutoRTFM::TOpenWrapper{MySharedPtr}]
 //    {
 //        WrappedSharedPtr.Object->DoSomething();
 //    });
