[AutoRTFM] Remove standard library fallback path

We've not had an ensure fire for some time, and so this fallback path should be safe to remove.

diff --git a/Engine/Source/Runtime/AutoRTFM/Private/FunctionMap.cpp b/Engine/Source/Runtime/AutoRTFM/Private/FunctionMap.cpp
--- a/Engine/Source/Runtime/AutoRTFM/Private/FunctionMap.cpp
+++ b/Engine/Source/Runtime/AutoRTFM/Private/FunctionMap.cpp
@@ -7,7 +7,6 @@
 #include "AutoRTFM.h"
 #include "BuildMacros.h"
 #include "FunctionMapInlines.h"
-#include "StdLibSupport.h"
 #include "Utils.h"
 
 #include <cinttypes>
@@ -330,13 +329,6 @@ void* FunctionMapLookupExhaustive(void* OpenFn, const char* Where)
 		}
 	}
 
-	// Check the temporary, old standard library mappings.
-	// This will be removed in the near future.
-	if (void* ClosedFn = LookupStdLibFunctionAndEnsureIfFound(OpenFn))
-	{
-		return ClosedFn;
-	}
-
 	AUTORTFM_MUST_TAIL return FunctionMapReportError(OpenFn, Where);
 }
 
diff --git a/Engine/Source/Runtime/AutoRTFM/Private/StdLibSupport.cpp b/Engine/Source/Runtime/AutoRTFM/Private/StdLibSupport.cpp
--- a/Engine/Source/Runtime/AutoRTFM/Private/StdLibSupport.cpp
+++ b/Engine/Source/Runtime/AutoRTFM/Private/StdLibSupport.cpp
@@ -2,7 +2,6 @@
 
 #if (defined(__AUTORTFM) && __AUTORTFM)
 
-#include "StdLibSupport.h"
 #include "AutoRTFM.h"
 #include "BuildMacros.h"
 #include "ContextInlines.h"
@@ -991,366 +990,6 @@ extern "C" AUTORTFM_INTERNAL UE_AUTORTFM_API void RTFM_ZdaPvmSt11align_val_t(voi
 // ^^^^^^^^^^^^^^^^^^^^^^^^^^^^^^ AUTORTFM_PLATFORM_LINUX ^^^^^^^^^^^^^^^^^^^^^^^^^^^^^^
 #endif
 
-// Do not add new entries to this function as it will be removed once we're
-// confident that the AEM files entirely cover these mappings.
-void* LookupStdLibFunctionAndEnsureIfFound(void* Open)
-{
-#define UE_AUTORTFM_MAP_OPEN_TO_CLOSED(Open, Closed) std::make_pair(reinterpret_cast<void*>(Open), reinterpret_cast<void*>(Closed))
-#define UE_AUTORTFM_MAP_OPEN_TO_CLOSED_OVERLOADED(Signature, Open, Closed) UE_AUTORTFM_MAP_OPEN_TO_CLOSED(static_cast< std::add_pointer_t<Signature> >(Open), static_cast< std::add_pointer_t<Signature> >(Closed))
-#define UE_AUTORTFM_MAP_OPEN_TO_SELF(Open) UE_AUTORTFM_MAP_OPEN_TO_CLOSED(Open, Open)
-#define UE_AUTORTFM_MAP_OPEN_TO_SELF_OVERLOADED(Signature, Open) UE_AUTORTFM_MAP_OPEN_TO_CLOSED(static_cast< std::add_pointer_t<Signature> >(Open), static_cast< std::add_pointer_t<Signature> >(Open))
-
-	static std::unordered_map<void*, void*> Mappings {
-
-#if AUTORTFM_ASAN_ENABLED
-	UE_AUTORTFM_MAP_OPEN_TO_SELF(__asan_addr_is_in_fake_stack),
-	UE_AUTORTFM_MAP_OPEN_TO_SELF(__asan_get_current_fake_stack),
-	UE_AUTORTFM_MAP_OPEN_TO_SELF(__asan_poison_memory_region),
-#endif
-
-#if AUTORTFM_PLATFORM_WINDOWS
-	UE_AUTORTFM_MAP_OPEN_TO_CLOSED(_strtoi64, RTFM__strtoi64),
-	UE_AUTORTFM_MAP_OPEN_TO_CLOSED(_wcstoi64, RTFM__wcstoi64),
-	UE_AUTORTFM_MAP_OPEN_TO_CLOSED(_wcstoui64, RTFM__wcstoui64),
-	UE_AUTORTFM_MAP_OPEN_TO_CLOSED(_wcstod_l, RTFM__wcstod_l),
-	UE_AUTORTFM_MAP_OPEN_TO_CLOSED(_wcstof_l, RTFM__wcstof_l),
-	UE_AUTORTFM_MAP_OPEN_TO_CLOSED(_wcstold_l, RTFM__wcstold_l),
-	UE_AUTORTFM_MAP_OPEN_TO_CLOSED(_wcstol_l, RTFM__wcstol_l),
-	UE_AUTORTFM_MAP_OPEN_TO_CLOSED(_wcstoll_l, RTFM__wcstoll_l),
-	UE_AUTORTFM_MAP_OPEN_TO_CLOSED(_wcstoul_l, RTFM__wcstoul_l),
-	UE_AUTORTFM_MAP_OPEN_TO_CLOSED(_wcstoull_l, RTFM__wcstoull_l),
-	UE_AUTORTFM_MAP_OPEN_TO_CLOSED(_wcstoi64_l, RTFM__wcstoi64_l),
-	UE_AUTORTFM_MAP_OPEN_TO_CLOSED(_wcstoui64_l, RTFM__wcstoui64_l),
-	UE_AUTORTFM_MAP_OPEN_TO_SELF(_wtof),
-	UE_AUTORTFM_MAP_OPEN_TO_SELF(_wtof),
-	UE_AUTORTFM_MAP_OPEN_TO_SELF(_wtol),
-	UE_AUTORTFM_MAP_OPEN_TO_SELF(_wtoll),
-	UE_AUTORTFM_MAP_OPEN_TO_SELF(_wtoi),
-	UE_AUTORTFM_MAP_OPEN_TO_SELF(_wtoi64),
-	UE_AUTORTFM_MAP_OPEN_TO_SELF(_wtof_l),
-	UE_AUTORTFM_MAP_OPEN_TO_SELF(_wtol_l),
-	UE_AUTORTFM_MAP_OPEN_TO_SELF(_wtoll_l),
-	UE_AUTORTFM_MAP_OPEN_TO_SELF(_wtoi_l),
-	UE_AUTORTFM_MAP_OPEN_TO_SELF(_wtoi64_l),
-	UE_AUTORTFM_MAP_OPEN_TO_CLOSED(__acrt_iob_func, RTFM___acrt_iob_func),
-	UE_AUTORTFM_MAP_OPEN_TO_CLOSED(__stdio_common_vfprintf, RTFM___stdio_common_vfprintf),
-	UE_AUTORTFM_MAP_OPEN_TO_CLOSED(__stdio_common_vsprintf, RTFM___stdio_common_vsprintf),
-	UE_AUTORTFM_MAP_OPEN_TO_CLOSED(__stdio_common_vswprintf, RTFM___stdio_common_vswprintf),
-	UE_AUTORTFM_MAP_OPEN_TO_CLOSED(__stdio_common_vfwprintf, RTFM___stdio_common_vfwprintf),
-	UE_AUTORTFM_MAP_OPEN_TO_SELF(_tcsncmp),
-	UE_AUTORTFM_MAP_OPEN_TO_SELF(_tcslen),
-	UE_AUTORTFM_MAP_OPEN_TO_SELF(_tcsnlen),
-	UE_AUTORTFM_MAP_OPEN_TO_SELF(_isnan),
-	UE_AUTORTFM_MAP_OPEN_TO_SELF(_fdtest),
-	UE_AUTORTFM_MAP_OPEN_TO_SELF(_dtest),
-	UE_AUTORTFM_MAP_OPEN_TO_SELF(_ldtest),
-	UE_AUTORTFM_MAP_OPEN_TO_SELF(_finite),
-	UE_AUTORTFM_MAP_OPEN_TO_SELF(IsDebuggerPresent),
-	UE_AUTORTFM_MAP_OPEN_TO_SELF(GetSystemTime),
-	UE_AUTORTFM_MAP_OPEN_TO_SELF(QueryPerformanceCounter),
-	UE_AUTORTFM_MAP_OPEN_TO_SELF(QueryPerformanceFrequency),
-	UE_AUTORTFM_MAP_OPEN_TO_SELF(GetCurrentThreadId),
-	UE_AUTORTFM_MAP_OPEN_TO_SELF(GetCurrentProcessId),
-	UE_AUTORTFM_MAP_OPEN_TO_SELF(TlsGetValue),
-	UE_AUTORTFM_MAP_OPEN_TO_SELF(FlsGetValue),
-	UE_AUTORTFM_MAP_OPEN_TO_SELF(GetLocalTime),
-	UE_AUTORTFM_MAP_OPEN_TO_SELF(GetFileAttributesW),
-	UE_AUTORTFM_MAP_OPEN_TO_CLOSED(TlsSetValue, RTFM_TlsSetValue),
-	UE_AUTORTFM_MAP_OPEN_TO_CLOSED(FlsSetValue, RTFM_FlsSetValue),
-	UE_AUTORTFM_MAP_OPEN_TO_SELF(_Query_perf_frequency),
-	UE_AUTORTFM_MAP_OPEN_TO_SELF(_Query_perf_counter),
-#endif
-
-#if AUTORTFM_PLATFORM_WINDOWS && AUTORTFM_EXCEPTIONS_ENABLED
-	UE_AUTORTFM_MAP_OPEN_TO_SELF(_CxxThrowException),
-#endif
-
-#if AUTORTFM_PLATFORM_LINUX && AUTORTFM_EXCEPTIONS_ENABLED
-	UE_AUTORTFM_MAP_OPEN_TO_SELF(__cxa_allocate_exception),
-	UE_AUTORTFM_MAP_OPEN_TO_SELF(__cxa_throw),
-	UE_AUTORTFM_MAP_OPEN_TO_SELF(__cxa_begin_catch),
-	UE_AUTORTFM_MAP_OPEN_TO_SELF(__cxa_end_catch),
-#endif
-
-#if AUTORTFM_PLATFORM_LINUX
-	UE_AUTORTFM_MAP_OPEN_TO_SELF(getpid),
-	UE_AUTORTFM_MAP_OPEN_TO_SELF(clock_gettime),
-	UE_AUTORTFM_MAP_OPEN_TO_SELF(gettimeofday),
-	UE_AUTORTFM_MAP_OPEN_TO_SELF(gmtime_r),
-	UE_AUTORTFM_MAP_OPEN_TO_SELF(localtime_r),
-	UE_AUTORTFM_MAP_OPEN_TO_SELF(bcmp),
-	UE_AUTORTFM_MAP_OPEN_TO_SELF(pthread_getspecific),
-	UE_AUTORTFM_MAP_OPEN_TO_SELF(pthread_self),
-	UE_AUTORTFM_MAP_OPEN_TO_CLOSED(strtof32, RTFM_strtof32),
-	UE_AUTORTFM_MAP_OPEN_TO_CLOSED(strtof64, RTFM_strtof64),
-	UE_AUTORTFM_MAP_OPEN_TO_CLOSED(stat, RTFM_stat),
-	UE_AUTORTFM_MAP_OPEN_TO_CLOSED(fstat, RTFM_fstat),
-	UE_AUTORTFM_MAP_OPEN_TO_CLOSED(__xstat, RTFM___xstat),
-	UE_AUTORTFM_MAP_OPEN_TO_CLOSED(__fxstat, RTFM___fxstat),
-#endif
-
-// Linux requires mapping of the C functions, which do not use const on their return pointers.
-// See comment about LibCxx builtin functions above.
-#if AUTORTFM_PLATFORM_LINUX
-	UE_AUTORTFM_MAP_OPEN_TO_SELF_OVERLOADED(char* (const char*, int), strchr),
-	UE_AUTORTFM_MAP_OPEN_TO_SELF_OVERLOADED(char* (const char*, int), strrchr),
-	UE_AUTORTFM_MAP_OPEN_TO_SELF_OVERLOADED(char* (const char*, const char*), strstr),
-#else
-	UE_AUTORTFM_MAP_OPEN_TO_SELF_OVERLOADED(const char* (const char*, int), strchr),
-	UE_AUTORTFM_MAP_OPEN_TO_SELF_OVERLOADED(const char* (const char*, int), strrchr),
-	UE_AUTORTFM_MAP_OPEN_TO_SELF_OVERLOADED(const char* (const char*, const char*), strstr),
-#endif
-
-#if AUTORTFM_PLATFORM_LINUX
-	UE_AUTORTFM_MAP_OPEN_TO_CLOSED(_Znwm, RTFM_Znwm),
-	UE_AUTORTFM_MAP_OPEN_TO_CLOSED(_Znam, RTFM_Znam),
-	UE_AUTORTFM_MAP_OPEN_TO_CLOSED(_ZnwmRKSt9nothrow_t, RTFM_ZnwmRKSt9nothrow_t),
-	UE_AUTORTFM_MAP_OPEN_TO_CLOSED(_ZnamRKSt9nothrow_t, RTFM_ZnamRKSt9nothrow_t),
-	UE_AUTORTFM_MAP_OPEN_TO_CLOSED(_ZnwmSt11align_val_t, RTFM_ZnwmSt11align_val_t),
-	UE_AUTORTFM_MAP_OPEN_TO_CLOSED(_ZnamSt11align_val_t, RTFM_ZnamSt11align_val_t),
-	UE_AUTORTFM_MAP_OPEN_TO_CLOSED(_ZnwmSt11align_val_tRKSt9nothrow_t, RTFM_ZnwmSt11align_val_tRKSt9nothrow_t),
-	UE_AUTORTFM_MAP_OPEN_TO_CLOSED(_ZnamSt11align_val_tRKSt9nothrow_t, RTFM_ZnamSt11align_val_tRKSt9nothrow_t),
-	UE_AUTORTFM_MAP_OPEN_TO_CLOSED(_ZdlPv, RTFM_ZdlPv),
-	UE_AUTORTFM_MAP_OPEN_TO_CLOSED(_ZdaPv, RTFM_ZdaPv),
-	UE_AUTORTFM_MAP_OPEN_TO_CLOSED(_ZdlPvRKSt9nothrow_t, RTFM_ZdlPvRKSt9nothrow_t),
-	UE_AUTORTFM_MAP_OPEN_TO_CLOSED(_ZdaPvRKSt9nothrow_t, RTFM_ZdaPvRKSt9nothrow_t),
-	UE_AUTORTFM_MAP_OPEN_TO_CLOSED(_ZdlPvm, RTFM_ZdlPvm),
-	UE_AUTORTFM_MAP_OPEN_TO_CLOSED(_ZdaPvm, RTFM_ZdaPvm),
-	UE_AUTORTFM_MAP_OPEN_TO_CLOSED(_ZdlPvSt11align_val_t, RTFM_ZdlPvSt11align_val_t),
-	UE_AUTORTFM_MAP_OPEN_TO_CLOSED(_ZdaPvSt11align_val_t, RTFM_ZdaPvSt11align_val_t),
-	UE_AUTORTFM_MAP_OPEN_TO_CLOSED(_ZdlPvSt11align_val_tRKSt9nothrow_t, RTFM_ZdlPvSt11align_val_tRKSt9nothrow_t),
-	UE_AUTORTFM_MAP_OPEN_TO_CLOSED(_ZdaPvSt11align_val_tRKSt9nothrow_t, RTFM_ZdaPvSt11align_val_tRKSt9nothrow_t),
-	UE_AUTORTFM_MAP_OPEN_TO_CLOSED(_ZdlPvmSt11align_val_t, RTFM_ZdlPvmSt11align_val_t),
-	UE_AUTORTFM_MAP_OPEN_TO_CLOSED(_ZdaPvmSt11align_val_t, RTFM_ZdaPvmSt11align_val_t),
-#endif
-
-	UE_AUTORTFM_MAP_OPEN_TO_CLOSED(memcpy, RTFM_memcpy),
-	UE_AUTORTFM_MAP_OPEN_TO_CLOSED(memmove, RTFM_memmove),
-	UE_AUTORTFM_MAP_OPEN_TO_CLOSED(memset, RTFM_memset),
-	UE_AUTORTFM_MAP_OPEN_TO_CLOSED(malloc, RTFM_malloc),
-	UE_AUTORTFM_MAP_OPEN_TO_CLOSED(calloc, RTFM_calloc),
-	UE_AUTORTFM_MAP_OPEN_TO_CLOSED(free, RTFM_free),
-	UE_AUTORTFM_MAP_OPEN_TO_CLOSED(realloc, RTFM_realloc),
-	UE_AUTORTFM_MAP_OPEN_TO_CLOSED(strcpy, RTFM_strcpy),
-	UE_AUTORTFM_MAP_OPEN_TO_CLOSED(strncpy, RTFM_strncpy),
-	UE_AUTORTFM_MAP_OPEN_TO_CLOSED(strcat, RTFM_strcat),
-	UE_AUTORTFM_MAP_OPEN_TO_CLOSED(strncat, RTFM_strncat),
-	UE_AUTORTFM_MAP_OPEN_TO_SELF(memcmp),
-	UE_AUTORTFM_MAP_OPEN_TO_SELF(strcmp),
-	UE_AUTORTFM_MAP_OPEN_TO_SELF(strncmp),
-	UE_AUTORTFM_MAP_OPEN_TO_SELF(strlen),
-	UE_AUTORTFM_MAP_OPEN_TO_SELF(atof),
-	UE_AUTORTFM_MAP_OPEN_TO_SELF(atoi),
-	UE_AUTORTFM_MAP_OPEN_TO_SELF_OVERLOADED(const wchar_t* (const wchar_t*, wchar_t), wcschr),
-	UE_AUTORTFM_MAP_OPEN_TO_SELF_OVERLOADED(wchar_t* (wchar_t*, wchar_t), wcschr),
-	UE_AUTORTFM_MAP_OPEN_TO_SELF_OVERLOADED(wchar_t* (wchar_t*, const wchar_t*), wcsstr),
-	UE_AUTORTFM_MAP_OPEN_TO_SELF_OVERLOADED(const wchar_t*(const wchar_t*, const wchar_t*), wcsstr),
-	UE_AUTORTFM_MAP_OPEN_TO_SELF_OVERLOADED(wchar_t*(wchar_t*, wchar_t), wcsrchr),
-	UE_AUTORTFM_MAP_OPEN_TO_SELF_OVERLOADED(const wchar_t*(const wchar_t*, wchar_t), wcsrchr),
-	UE_AUTORTFM_MAP_OPEN_TO_SELF(wcscmp),
-	UE_AUTORTFM_MAP_OPEN_TO_SELF(wcslen),
-	UE_AUTORTFM_MAP_OPEN_TO_CLOSED(strtol, RTFM_strtol),
-	UE_AUTORTFM_MAP_OPEN_TO_CLOSED(strtoll, RTFM_strtoll),
-	UE_AUTORTFM_MAP_OPEN_TO_CLOSED(strtoul, RTFM_strtoul),
-	UE_AUTORTFM_MAP_OPEN_TO_CLOSED(strtoull, RTFM_strtoull),
-	UE_AUTORTFM_MAP_OPEN_TO_CLOSED(strtof, RTFM_strtof),
-	UE_AUTORTFM_MAP_OPEN_TO_CLOSED(strtod, RTFM_strtod),
-	UE_AUTORTFM_MAP_OPEN_TO_CLOSED(wcstod, RTFM_wcstod),
-	UE_AUTORTFM_MAP_OPEN_TO_CLOSED(wcstof, RTFM_wcstof),
-	UE_AUTORTFM_MAP_OPEN_TO_CLOSED(wcstold, RTFM_wcstold),
-	UE_AUTORTFM_MAP_OPEN_TO_CLOSED(wcstol, RTFM_wcstol),
-	UE_AUTORTFM_MAP_OPEN_TO_CLOSED(wcstoll, RTFM_wcstoll),
-	UE_AUTORTFM_MAP_OPEN_TO_CLOSED(wcstoul, RTFM_wcstoul),
-	UE_AUTORTFM_MAP_OPEN_TO_CLOSED(wcstoull, RTFM_wcstoull),
-	UE_AUTORTFM_MAP_OPEN_TO_CLOSED_OVERLOADED(
-		std::to_chars_result(char* First, char* Last, int Value, int Base),
-		std::to_chars, RTFM__ToChars_i),
-	UE_AUTORTFM_MAP_OPEN_TO_CLOSED_OVERLOADED(
-		std::to_chars_result(char* First, char* Last, unsigned int Value, int Base),
-		std::to_chars, RTFM__ToChars_u),
-	UE_AUTORTFM_MAP_OPEN_TO_CLOSED_OVERLOADED(
-		std::to_chars_result(char* First, char* Last, int8_t Value, int Base),
-		std::to_chars, RTFM__ToChars_i8),
-	UE_AUTORTFM_MAP_OPEN_TO_CLOSED_OVERLOADED(
-		std::to_chars_result(char* First, char* Last, uint8_t Value, int Base),
-		std::to_chars, RTFM__ToChars_u8),
-	UE_AUTORTFM_MAP_OPEN_TO_CLOSED_OVERLOADED(
-		std::to_chars_result(char* First, char* Last, int16_t Value, int Base),
-		std::to_chars, RTFM__ToChars_i16),
-	UE_AUTORTFM_MAP_OPEN_TO_CLOSED_OVERLOADED(
-		std::to_chars_result(char* First, char* Last, uint16_t Value, int Base),
-		std::to_chars, RTFM__ToChars_u16),
-	UE_AUTORTFM_MAP_OPEN_TO_CLOSED_OVERLOADED(
-		std::to_chars_result(char* First, char* Last, int32_t Value, int Base),
-		std::to_chars, RTFM__ToChars_i32),
-	UE_AUTORTFM_MAP_OPEN_TO_CLOSED_OVERLOADED(
-		std::to_chars_result(char* First, char* Last, uint32_t Value, int Base),
-		std::to_chars, RTFM__ToChars_u32),
-	UE_AUTORTFM_MAP_OPEN_TO_CLOSED_OVERLOADED(
-		std::to_chars_result(char* First, char* Last, int64_t Value, int Base),
-		std::to_chars, RTFM__ToChars_i64),
-	UE_AUTORTFM_MAP_OPEN_TO_CLOSED_OVERLOADED(
-		std::to_chars_result(char* First, char* Last, uint64_t Value, int Base),
-		std::to_chars, RTFM__ToChars_u64),
-	UE_AUTORTFM_MAP_OPEN_TO_CLOSED_OVERLOADED(
-		std::to_chars_result(char* First, char* Last, float Value),
-		std::to_chars, RTFM__ToChars_f),
-	UE_AUTORTFM_MAP_OPEN_TO_CLOSED_OVERLOADED(
-		std::to_chars_result(char* First, char* Last, float Value, std::chars_format Format),
-		std::to_chars, RTFM__ToChars_f_F),
-	UE_AUTORTFM_MAP_OPEN_TO_CLOSED_OVERLOADED(
-		std::to_chars_result(char* First, char* Last, float Value, std::chars_format Format, int Precision),
-		std::to_chars, RTFM__ToChars_f_FP),
-	UE_AUTORTFM_MAP_OPEN_TO_CLOSED_OVERLOADED(
-		std::to_chars_result(char* First, char* Last, double Value),
-		std::to_chars, RTFM__ToChars_d),
-	UE_AUTORTFM_MAP_OPEN_TO_CLOSED_OVERLOADED(
-		std::to_chars_result(char* First, char* Last, double Value, std::chars_format Format),
-		std::to_chars, RTFM__ToChars_d_F),
-	UE_AUTORTFM_MAP_OPEN_TO_CLOSED_OVERLOADED(
-		std::to_chars_result(char* First, char* Last, double Value, std::chars_format Format, int Precision),
-		std::to_chars, RTFM__ToChars_d_FP),
-	UE_AUTORTFM_MAP_OPEN_TO_CLOSED_OVERLOADED(
-		std::to_chars_result(char* First, char* Last, long double Value),
-		std::to_chars, RTFM__ToChars_ld),
-	UE_AUTORTFM_MAP_OPEN_TO_CLOSED_OVERLOADED(
-		std::to_chars_result(char* First, char* Last, long double Value, std::chars_format Format),
-		std::to_chars, RTFM__ToChars_ld_F),
-	UE_AUTORTFM_MAP_OPEN_TO_CLOSED_OVERLOADED(
-		std::to_chars_result(char* First, char* Last, long double Value, std::chars_format Format, int Precision),
-		std::to_chars, RTFM__ToChars_ld_FP),
-
-	UE_AUTORTFM_MAP_OPEN_TO_SELF(iswupper),
-	UE_AUTORTFM_MAP_OPEN_TO_SELF(iswlower),
-	UE_AUTORTFM_MAP_OPEN_TO_SELF(iswalpha),
-	UE_AUTORTFM_MAP_OPEN_TO_SELF(iswgraph),
-	UE_AUTORTFM_MAP_OPEN_TO_SELF(iswprint),
-	UE_AUTORTFM_MAP_OPEN_TO_SELF(iswpunct),
-	UE_AUTORTFM_MAP_OPEN_TO_SELF(iswalnum),
-	UE_AUTORTFM_MAP_OPEN_TO_SELF(iswdigit),
-	UE_AUTORTFM_MAP_OPEN_TO_SELF(iswxdigit),
-	UE_AUTORTFM_MAP_OPEN_TO_SELF(iswspace),
-	UE_AUTORTFM_MAP_OPEN_TO_SELF(iswcntrl),
-
-	UE_AUTORTFM_MAP_OPEN_TO_SELF_OVERLOADED(float (float), sqrt),
-	UE_AUTORTFM_MAP_OPEN_TO_SELF_OVERLOADED(double (double), sqrt),
-	UE_AUTORTFM_MAP_OPEN_TO_SELF_OVERLOADED(long double (long double), sqrt),
-	UE_AUTORTFM_MAP_OPEN_TO_SELF_OVERLOADED(float (float), sin),
-	UE_AUTORTFM_MAP_OPEN_TO_SELF_OVERLOADED(double (double), sin),
-	UE_AUTORTFM_MAP_OPEN_TO_SELF_OVERLOADED(long double (long double), sin),
-	UE_AUTORTFM_MAP_OPEN_TO_SELF_OVERLOADED(float (float), cos),
-	UE_AUTORTFM_MAP_OPEN_TO_SELF_OVERLOADED(double (double), cos),
-	UE_AUTORTFM_MAP_OPEN_TO_SELF_OVERLOADED(long double (long double), cos),
-	UE_AUTORTFM_MAP_OPEN_TO_SELF_OVERLOADED(float (float), tan),
-	UE_AUTORTFM_MAP_OPEN_TO_SELF_OVERLOADED(double (double), tan),
-	UE_AUTORTFM_MAP_OPEN_TO_SELF_OVERLOADED(long double (long double), tan),
-	UE_AUTORTFM_MAP_OPEN_TO_SELF_OVERLOADED(float (float), asin),
-	UE_AUTORTFM_MAP_OPEN_TO_SELF_OVERLOADED(double (double), asin),
-	UE_AUTORTFM_MAP_OPEN_TO_SELF_OVERLOADED(long double (long double), asin),
-	UE_AUTORTFM_MAP_OPEN_TO_SELF_OVERLOADED(float (float), acos),
-	UE_AUTORTFM_MAP_OPEN_TO_SELF_OVERLOADED(double (double), acos),
-	UE_AUTORTFM_MAP_OPEN_TO_SELF_OVERLOADED(long double (long double), acos),
-	UE_AUTORTFM_MAP_OPEN_TO_SELF_OVERLOADED(float (float), atan),
-	UE_AUTORTFM_MAP_OPEN_TO_SELF_OVERLOADED(double (double), atan),
-	UE_AUTORTFM_MAP_OPEN_TO_SELF_OVERLOADED(long double (long double), atan),
-	UE_AUTORTFM_MAP_OPEN_TO_SELF_OVERLOADED(float (float, float), atan2),
-	UE_AUTORTFM_MAP_OPEN_TO_SELF_OVERLOADED(double (double, double), atan2),
-	UE_AUTORTFM_MAP_OPEN_TO_SELF_OVERLOADED(long double (long double, long double), atan2),
-	UE_AUTORTFM_MAP_OPEN_TO_SELF_OVERLOADED(float (float), sinh),
-	UE_AUTORTFM_MAP_OPEN_TO_SELF_OVERLOADED(double (double), sinh),
-	UE_AUTORTFM_MAP_OPEN_TO_SELF_OVERLOADED(long double (long double), sinh),
-	UE_AUTORTFM_MAP_OPEN_TO_SELF_OVERLOADED(float (float), cosh),
-	UE_AUTORTFM_MAP_OPEN_TO_SELF_OVERLOADED(double (double), cosh),
-	UE_AUTORTFM_MAP_OPEN_TO_SELF_OVERLOADED(long double (long double), cosh),
-	UE_AUTORTFM_MAP_OPEN_TO_SELF_OVERLOADED(float (float), tanh),
-	UE_AUTORTFM_MAP_OPEN_TO_SELF_OVERLOADED(double (double), tanh),
-	UE_AUTORTFM_MAP_OPEN_TO_SELF_OVERLOADED(long double (long double), tanh),
-	UE_AUTORTFM_MAP_OPEN_TO_SELF_OVERLOADED(float (float), asinh),
-	UE_AUTORTFM_MAP_OPEN_TO_SELF_OVERLOADED(double (double), asinh),
-	UE_AUTORTFM_MAP_OPEN_TO_SELF_OVERLOADED(long double (long double), asinh),
-	UE_AUTORTFM_MAP_OPEN_TO_SELF_OVERLOADED(float (float), acosh),
-	UE_AUTORTFM_MAP_OPEN_TO_SELF_OVERLOADED(double (double), acosh),
-	UE_AUTORTFM_MAP_OPEN_TO_SELF_OVERLOADED(long double (long double), acosh),
-	UE_AUTORTFM_MAP_OPEN_TO_SELF_OVERLOADED(float (float), atanh),
-	UE_AUTORTFM_MAP_OPEN_TO_SELF_OVERLOADED(double (double), atanh),
-	UE_AUTORTFM_MAP_OPEN_TO_SELF_OVERLOADED(long double (long double), atanh),
-	UE_AUTORTFM_MAP_OPEN_TO_SELF_OVERLOADED(float (float), exp),
-	UE_AUTORTFM_MAP_OPEN_TO_SELF_OVERLOADED(double (double), exp),
-	UE_AUTORTFM_MAP_OPEN_TO_SELF_OVERLOADED(long double (long double), exp),
-	UE_AUTORTFM_MAP_OPEN_TO_SELF_OVERLOADED(float (float), log),
-	UE_AUTORTFM_MAP_OPEN_TO_SELF_OVERLOADED(double (double), log),
-	UE_AUTORTFM_MAP_OPEN_TO_SELF_OVERLOADED(long double (long double), log),
-	UE_AUTORTFM_MAP_OPEN_TO_SELF_OVERLOADED(float (float, float), pow),
-	UE_AUTORTFM_MAP_OPEN_TO_SELF_OVERLOADED(double (double, double), pow),
-	UE_AUTORTFM_MAP_OPEN_TO_SELF_OVERLOADED(long double (long double, long double), pow),
-	UE_AUTORTFM_MAP_OPEN_TO_SELF_OVERLOADED(long long (float), llrint),
-	UE_AUTORTFM_MAP_OPEN_TO_SELF_OVERLOADED(long long (double), llrint),
-	UE_AUTORTFM_MAP_OPEN_TO_SELF_OVERLOADED(long long (long double), llrint),
-	UE_AUTORTFM_MAP_OPEN_TO_SELF_OVERLOADED(float (float, float), fmod),
-	UE_AUTORTFM_MAP_OPEN_TO_SELF_OVERLOADED(double (double, double), fmod),
-	UE_AUTORTFM_MAP_OPEN_TO_SELF_OVERLOADED(long double (long double, long double), fmod),
-	UE_AUTORTFM_MAP_OPEN_TO_SELF_OVERLOADED(float (float, float*), modf),
-	UE_AUTORTFM_MAP_OPEN_TO_SELF_OVERLOADED(double (double, double*), modf),
-	UE_AUTORTFM_MAP_OPEN_TO_SELF_OVERLOADED(long double (long double, long double*), modf),
-	UE_AUTORTFM_MAP_OPEN_TO_SELF_OVERLOADED(float (float, float), powf),
-	UE_AUTORTFM_MAP_OPEN_TO_SELF_OVERLOADED(double (double, double), pow),
-	UE_AUTORTFM_MAP_OPEN_TO_SELF(sqrtf),
-	UE_AUTORTFM_MAP_OPEN_TO_SELF(sinf),
-	UE_AUTORTFM_MAP_OPEN_TO_SELF(cosf),
-	UE_AUTORTFM_MAP_OPEN_TO_SELF(tanf),
-	UE_AUTORTFM_MAP_OPEN_TO_SELF(asinf),
-	UE_AUTORTFM_MAP_OPEN_TO_SELF(asinhf),
-	UE_AUTORTFM_MAP_OPEN_TO_SELF(acosf),
-	UE_AUTORTFM_MAP_OPEN_TO_SELF(acoshf),
-	UE_AUTORTFM_MAP_OPEN_TO_SELF(atanf),
-	UE_AUTORTFM_MAP_OPEN_TO_SELF(atanhf),
-	UE_AUTORTFM_MAP_OPEN_TO_SELF(atan2f),
-	UE_AUTORTFM_MAP_OPEN_TO_SELF(sinhf),
-	UE_AUTORTFM_MAP_OPEN_TO_SELF(coshf),
-	UE_AUTORTFM_MAP_OPEN_TO_SELF(tanhf),
-	UE_AUTORTFM_MAP_OPEN_TO_SELF(expf),
-	UE_AUTORTFM_MAP_OPEN_TO_SELF(logf),
-	UE_AUTORTFM_MAP_OPEN_TO_SELF(powf),
-	UE_AUTORTFM_MAP_OPEN_TO_SELF(llrintf),
-	UE_AUTORTFM_MAP_OPEN_TO_SELF(fmodf),
-	UE_AUTORTFM_MAP_OPEN_TO_SELF(fmodl),
-	UE_AUTORTFM_MAP_OPEN_TO_SELF(rand),
-	UE_AUTORTFM_MAP_OPEN_TO_SELF(modff),
-	UE_AUTORTFM_MAP_OPEN_TO_SELF(modfl),
-	UE_AUTORTFM_MAP_OPEN_TO_CLOSED(vsnprintf, RTFM_vsnprintf),
-	UE_AUTORTFM_MAP_OPEN_TO_CLOSED_OVERLOADED(
-		int (wchar_t*, size_t, const wchar_t*, va_list),
-		vswprintf, RTFM_vswprintf),
-	UE_AUTORTFM_MAP_OPEN_TO_CLOSED_OVERLOADED(
-		int(wchar_t*, size_t, wchar_t const*, ...),
-		swprintf, RTFM_swprintf),
-	UE_AUTORTFM_MAP_OPEN_TO_CLOSED(snprintf, RTFM_snprintf),
-	UE_AUTORTFM_MAP_OPEN_TO_CLOSED(printf, RTFM_printf),
-	UE_AUTORTFM_MAP_OPEN_TO_CLOSED(wprintf, RTFM_wprintf),
-	UE_AUTORTFM_MAP_OPEN_TO_SELF(putchar),
-	UE_AUTORTFM_MAP_OPEN_TO_SELF(puts),
-	UE_AUTORTFM_MAP_OPEN_TO_SELF(fflush),
-	UE_AUTORTFM_MAP_OPEN_TO_CLOSED(wcscpy, RTFM_wcscpy),
-	UE_AUTORTFM_MAP_OPEN_TO_CLOSED(wcsncpy, RTFM_wcsncpy),
-	UE_AUTORTFM_MAP_OPEN_TO_CLOSED(atexit, RTFM_atexit)
-
-	};
-
-	void* Closed = Mappings[Open];
-	AUTORTFM_ENSURE_MSG(!Closed,
-		"Closed function for '%s' found in LookupStdLibFunctionAndEnsureIfFound(), "
-		"but should have been declared in the external mapping files",
-		GetFunctionDescription(Open).c_str()
-	);
-	return Closed;
-}
-
-
-#ifdef _MSC_VER
-#pragma warning(pop)
-// END: Disable warning about deprecated STD C functions.
-#endif
-
 } // namespace AutoRTFM
 
 #endif // defined(__AUTORTFM) && __AUTORTFM
diff --git a/Engine/Source/Runtime/AutoRTFM/Private/StdLibSupport.h b/Engine/Source/Runtime/AutoRTFM/Private/StdLibSupport.h
deleted file mode 100644
