[AutoRTFM]
* Fixed iwyu pragma

diff --git a/Engine/Source/Runtime/AutoRTFM/Public/AutoRTFMDefines.h b/Engine/Source/Runtime/AutoRTFM/Public/AutoRTFMDefines.h
--- a/Engine/Source/Runtime/AutoRTFM/Public/AutoRTFMDefines.h
+++ b/Engine/Source/Runtime/AutoRTFM/Public/AutoRTFMDefines.h
@@ -1,6 +1,8 @@
 // Copyright Epic Games, Inc. All Rights Reserved.
 
-// IWYU pragma: private; include "AutoRTFM.h".
+// IWYU pragma: private
+// Include include "AutoRTFM.h" instead
+
 #pragma once
 
 #if (defined(__AUTORTFM) && __AUTORTFM)
