[AutoRTFM] Make AutoRTFM known hazards as easy as possible to spot in a crash log.

Previously, an AutoRTFM crash message gave very little information to a reader:
- Could not find function 00007FFF9709B430 'RtlInitializeSRWLock' where 'call'.

Now, we print out a large eye-catching banner, a detailed message explaining what needs to be fixed, and a call stack pointing out where the issue lies. This will make it easier for a developer to self-service an issue instead of immediately bringing it to our team.

This CL addresses FCriticalSection; in a followup, I intend to add messages for FRWLock, atomics, and other common failure modes.

diff --git a/Engine/Source/Programs/AutoRTFMTests/Private/Abort.cpp b/Engine/Source/Programs/AutoRTFMTests/Private/Abort.cpp
--- a/Engine/Source/Programs/AutoRTFMTests/Private/Abort.cpp
+++ b/Engine/Source/Programs/AutoRTFMTests/Private/Abort.cpp
@@ -3,7 +3,10 @@
 #include "AutoRTFM.h"
 #include "AutoRTFMTesting.h"
 #include "AutoRTFMTestUtils.h"
+#include "AutoRTFM/AutoRTFMUE.h"
 #include "Catch2Includes.h"
+#include "Misc/StringOutputDevice.h"
+#include "Misc/OutputDeviceRedirector.h"
 
 #include <map>
 #include <vector>
@@ -199,7 +202,7 @@ TEST_CASE("Abort.TransactionInCommittedOnComplete")
 				bDidSomething = true;
 			});
 
-			// The test harness will abort this transaction and retry it when `ShouldRetryNonNestedTransactions` 
+			// The test harness will abort this transaction and retry it when `ShouldRetryNonNestedTransactions`
 			// is set, which unfortunately muddles the test result here since we'd normally expect "transact during
 			// commit" here, not "transact during abort."
 			REQUIRE((InnerResult == AutoRTFM::ETransactionResult::AbortedByTransactDuringCommit ||
@@ -487,7 +490,7 @@ TEST_CASE("Abort.CascadeWithOnCompleteThroughFunctionAttribute")
 				});
 			});
 		}
-		
+
 		bool bTouched = false;
 		bool bOnCompleteCallbackWasExecuted = false;
 	};
@@ -868,7 +871,7 @@ TEST_CASE("Abort.PushOnAbortHandler_Order")
 				{
 					AutoRTFM::PushOnAbortHandler(&SAbortKey, [&Value]()
 					{
-						// Only if we are retrying on 
+						// Only if we are retrying on
 						REQUIRE(AutoRTFM::ForTheRuntime::ShouldRetryNestedTransactionsToo());
 					});
 				});
@@ -1354,3 +1357,68 @@ TEST_CASE("Abort.CallbackModifiesTransactionalMemory")
 		}
 	}
 }
+
+namespace
+{
+
+class FScopedErrorHintTest
+{
+	public:
+		FScopedErrorHintTest(FOutputDevice& InDevice) : Device(InDevice)
+		{
+			OriginalVerbosity = LogAutoRTFM.GetVerbosity();
+			GLog->AddOutputDevice(&Device);
+			LogAutoRTFM.SetVerbosity(ELogVerbosity::Display);
+		}
+
+		~FScopedErrorHintTest()
+		{
+			LogAutoRTFM.SetVerbosity(OriginalVerbosity);
+			GLog->RemoveOutputDevice(&Device);
+		}
+
+	private:
+		AutoRTFMTestUtils::FScopedInternalAbortAction Scoped1{AutoRTFM::ForTheRuntime::EAutoRTFMInternalAbortActionState::Abort};
+		AutoRTFMTestUtils::FScopedEnsureOnInternalAbort Scoped2{false};
+		AutoRTFMTestUtils::FScopedHazardLogSeverity Scoped3{autortfm_log_info};
+		ELogVerbosity::Type OriginalVerbosity;
+		FOutputDevice& Device;
+};
+
+}
+
+TEST_CASE("Abort.Language.ErrorHints.InitializeCriticalSection")
+{
+	TLogCategoryOutputDevice StringLog(LogAutoRTFM);
+	FScopedErrorHintTest TestWith(StringLog);
+
+	const AutoRTFM::ETransactionResult Result = AutoRTFM::Transact([&]
+	{
+		FCriticalSection CS;
+	});
+
+	REQUIRE(AutoRTFM::ETransactionResult::AbortedByLanguage == Result);
+
+	// Did we print our detailed error message?
+	REQUIRE(StringLog.Contains("A critical section was accessed while inside an AutoRTFM transaction."));
+	REQUIRE(StringLog.Contains("AutoRTFM/Documentation/README.md"));
+	REQUIRE(StringLog.Contains("Use the call stack below to find the FCriticalSection"));
+}
+
+TEST_CASE("Abort.Language.ErrorHints.EnterCriticalSection")
+{
+	TLogCategoryOutputDevice StringLog(LogAutoRTFM);
+	FScopedErrorHintTest TestWith(StringLog);
+
+	FCriticalSection CS;
+
+	const AutoRTFM::ETransactionResult Result = AutoRTFM::Transact([&]
+	{
+		CS.Lock();
+	});
+
+	REQUIRE(AutoRTFM::ETransactionResult::AbortedByLanguage == Result);
+
+	// Did we print our message?
+	REQUIRE(StringLog.Contains("A critical section was accessed while inside an AutoRTFM transaction."));
+}
diff --git a/Engine/Source/Programs/AutoRTFMTests/Private/AutoRTFMTestUtils.h b/Engine/Source/Programs/AutoRTFMTests/Private/AutoRTFMTestUtils.h
--- a/Engine/Source/Programs/AutoRTFMTests/Private/AutoRTFMTestUtils.h
+++ b/Engine/Source/Programs/AutoRTFMTests/Private/AutoRTFMTestUtils.h
@@ -141,7 +141,7 @@ static constexpr const std::string_view kMemoryModifiedInAbortHandlerDtorWarning
 
 static constexpr const std::string_view kMemoryModifiedInCommitHandlerDtorWarning =
 	"Memory modified in a transaction was also modified in the destructor of the on-commit handler:";
-	
+
 class FScopedEnsureOnInternalAbort final
 {
 public:
@@ -184,4 +184,25 @@ private:
 	const AutoRTFM::ForTheRuntime::EAutoRTFMInternalAbortActionState Original;
 };
 
+class FScopedHazardLogSeverity final
+{
+public:
+	FScopedHazardLogSeverity(autortfm_log_severity Severity)
+		: Original(AutoRTFM::ForTheRuntime::GetAutoRTFMHazardLogSeverity())
+	{
+		AutoRTFM::ForTheRuntime::SetAutoRTFMHazardLogSeverity(Severity);
+	}
+
+	~FScopedHazardLogSeverity()
+	{
+		AutoRTFM::ForTheRuntime::SetAutoRTFMHazardLogSeverity(Original);
+	}
+
+private:
+	FScopedHazardLogSeverity(const FScopedHazardLogSeverity&) = delete;
+	FScopedHazardLogSeverity& operator=(const FScopedHazardLogSeverity&) = delete;
+
+	const autortfm_log_severity Original;
+};
+
 }
diff --git a/Engine/Source/Runtime/AutoRTFM/Private/API.cpp b/Engine/Source/Runtime/AutoRTFM/Private/API.cpp
--- a/Engine/Source/Runtime/AutoRTFM/Private/API.cpp
+++ b/Engine/Source/Runtime/AutoRTFM/Private/API.cpp
@@ -26,7 +26,7 @@ namespace
 
 	int GAutoRTFMInternalAbortAction = AutoRTFM::ForTheRuntime::EAutoRTFMInternalAbortActionState::Crash;
 	int GAutoRTFMRetryTransactions = AutoRTFM::ForTheRuntime::EAutoRTFMRetryTransactionState::NoRetry;
-	
+
 	// Set the percentage chance [0..100] that AutoRTFM will be enabled.
 	// 100.0 means that AutoRTFM will always be enabled; 1.0 means that AutoRTFM has a 1% chance of being enabled.
 	// See `AutoRTFM::ForTheRuntime::CoinTossDisable` for implementation details.
@@ -36,10 +36,11 @@ namespace
 	// EMemoryValidationLevel::Default is a special enumerator that can be used
 	// in the public API to map to GDefaultMemoryValidationLevel.
 	AutoRTFM::EMemoryValidationLevel GMemoryValidationLevel = GDefaultMemoryValidationLevel;
-	
+
 	bool GMemoryValidationThrottlingEnabled = true;
 	bool GMemoryValidationStatisticsEnabled = true;
 	bool GAutoRTFMEnsureOnInternalAbort = true;
+	autortfm_log_severity GAutoRTFMHazardLogSeverity = autortfm_log_error;
 
 	// A linked-list of open->closed function tables populated by
 	// autortfm_register_open_to_closed_functions(). This is consumed by
@@ -119,7 +120,7 @@ namespace AutoRTFM
 				DoIgnoreLog(State, GAutoRTFMRuntimeEnabled);
 				return false;
 			}
-			
+
 			switch (GAutoRTFMRuntimeEnabled)
 			{
 			default:
@@ -262,6 +263,24 @@ namespace AutoRTFM
 #endif
 		}
 
+		autortfm_log_severity GetAutoRTFMHazardLogSeverity()
+		{
+#if UE_AUTORTFM
+			return GAutoRTFMHazardLogSeverity;
+#else
+			return autortfm_log_error;
+#endif
+		}
+
+		void SetAutoRTFMHazardLogSeverity(autortfm_log_severity Severity)
+		{
+#if UE_AUTORTFM
+			GAutoRTFMHazardLogSeverity = Severity;
+#else
+			UE_AUTORTFM_UNUSED(Severity);
+#endif
+		}
+
 		void SetRetryTransaction(EAutoRTFMRetryTransactionState State)
 		{
 #if UE_AUTORTFM
diff --git a/Engine/Source/Runtime/AutoRTFM/Private/FunctionMap.cpp b/Engine/Source/Runtime/AutoRTFM/Private/FunctionMap.cpp
--- a/Engine/Source/Runtime/AutoRTFM/Private/FunctionMap.cpp
+++ b/Engine/Source/Runtime/AutoRTFM/Private/FunctionMap.cpp
@@ -17,6 +17,12 @@
 #include <stdlib.h>
 #include <string.h>
 
+#if AUTORTFM_PLATFORM_WINDOWS
+	#include "WindowsHeader.h"
+#elif AUTORTFM_PLATFORM_LINUX
+	#include <pthread.h>
+#endif
+
 namespace AutoRTFM
 {
 
@@ -62,7 +68,7 @@ class FFunctionMap
 	{
 		return NumBitsInSizeT - __builtin_clzll(Value);
 	}
-	
+
 public:
 	// Obtains the FFunctionMap instance.
 	static FFunctionMap Get()
@@ -104,7 +110,7 @@ public:
 			{
 				continue;
 			}
-	
+
 			// TODO: If hash collisions are slow, we could have a secondary stronger hash here that
 			// is used on the second iteration followed by +1 for the subsequent iterations.
 			for (size_t Index = FunctionPtrHash(OldEntry->Open, NewCapacity2n);; Index = (Index + 1) & NewMap->IndexMask)
@@ -132,7 +138,7 @@ public:
 		{
 			autortfm_open_to_closed_mapping* Mapping = Map->Entries + Index;
 
-			if (!Mapping->Open) 
+			if (!Mapping->Open)
 			{
 				Mapping->Open = Open;
 				Mapping->Closed = Closed;
@@ -140,19 +146,19 @@ public:
 				return;
 			}
 
-			if (Mapping->Open == Open) 
+			if (Mapping->Open == Open)
 			{
 				Mapping->Closed = Closed;
 				return;
 			}
 		}
 	}
-	
+
 	// Looks up the closed function from the open function pointer.
 	// Returns nullptr if the mapping is not found.
 	void* Lookup(void* OpenFn)
 	{
-		for (size_t Index = FunctionPtrHash(OpenFn, Map->Capacity2n);; Index = (Index + 1) & Map->IndexMask) 
+		for (size_t Index = FunctionPtrHash(OpenFn, Map->Capacity2n);; Index = (Index + 1) & Map->IndexMask)
 		{
 			autortfm_open_to_closed_mapping* Entry = Map->Entries + Index;
 
@@ -173,11 +179,11 @@ public:
 		size_t MapSlots = 1llu << Map->Capacity2n;
 		AUTORTFM_LOG("Function Map Stats");
 		AUTORTFM_LOG("==================");
-		AUTORTFM_LOG("Occupancy: %zu entries in %zu slots (load factor %g%%)", 
-					 Map->EntryCount, 
+		AUTORTFM_LOG("Occupancy: %zu entries in %zu slots (load factor %g%%)",
+					 Map->EntryCount,
 					 MapSlots,
 					 static_cast<double>(Map->EntryCount * 1000llu / MapSlots) / 10.0);
-	
+
 		std::map<uintptr_t, size_t> CollisionMap; // <probes required to find the function, number of functions>
 		for (const autortfm_open_to_closed_mapping* Entry = Map->Entries, *FinalEntry = Map->Entries + Map->IndexMask; Entry <= FinalEntry; ++Entry)
 		{
@@ -186,17 +192,17 @@ public:
 				uintptr_t ActualIndex = std::distance<const autortfm_open_to_closed_mapping*>(Map->Entries, Entry);
 				uintptr_t IdealIndex = FunctionPtrHash(Entry->Open, Map->Capacity2n) & Map->IndexMask;
 				// We need to account for wraparound in this calculation.
-				uintptr_t Delta = (IdealIndex <= ActualIndex) 
+				uintptr_t Delta = (IdealIndex <= ActualIndex)
 					? ActualIndex - IdealIndex
 					: ActualIndex + MapSlots - IdealIndex;
-															   
+
 				CollisionMap[Delta] += 1;
 			}
 		}
-	
+
 		for (uintptr_t NumProbes = 0, HighestProbes = CollisionMap.rbegin()->first; NumProbes <= HighestProbes; ++NumProbes)
 		{
-			AUTORTFM_LOG("%2" PRIuPTR " probes: %zu functions", NumProbes, CollisionMap[NumProbes]); 
+			AUTORTFM_LOG("%2" PRIuPTR " probes: %zu functions", NumProbes, CollisionMap[NumProbes]);
 		}
 	}
 
@@ -210,7 +216,7 @@ private:
 	{
 		Mutex.lock();
 	}
-	
+
 	inline static uintptr_t FunctionPtrHash(void* FunctionPtr, size_t HashTableSize2n)
 	{
 		// Use a Fibonacci hash to fold our size_t value into a hash-appropriate value.
@@ -225,7 +231,7 @@ private:
 };
 
 // Attempts to lookup the "true" function from a dynamically linked import function thunk.
-// Returns the pointer to the true function, or nullptr if the function cannot 
+// Returns the pointer to the true function, or nullptr if the function cannot
 // be resolved (or is not an import thunk).
 inline void* FollowRelocation(void* Function)
 {
@@ -251,8 +257,67 @@ inline void* FollowRelocation(void* Function)
 	return nullptr;
 }
 
+template <typename... ArgList>
+AUTORTFM_DISABLE
+static void LogAutoRTFMHazard(autortfm_log_severity Severity, const char* Message, ArgList&&... Args)
+{
+	AutoRTFM::Log(__FILE__, __LINE__, /* ProgramCounter */ nullptr, Severity, Message, Args...);
+}
+
+AUTORTFM_DISABLE static const char* GetAutoRTFMHazardBanner()
+{
+	return R"(
+             _         _        ____ _____ _____ __  __   _   _                        _
+            / \  _   _| |_ ___ |  _ \_   _|  ___|  \/  | | | | | __ _ ______ _ _ __ __| |
+  _____    / _ \| | | | __/ _ \| |_) || | | |_  | |\/| | | |_| |/ _` |_  / _` | '__/ _` |  _____
+ |_____|  / ___ \ |_| | || (_) |  _ < | | |  _| | |  | | |  _  | (_| |/ / (_| | | | (_| | |_____|
+         /_/   \_\__,_|\__\___/|_| \_\|_| |_|   |_|  |_| |_| |_|\__,_/___\__,_|_|  \__,_|
+
+New to AutoRTFM? See Engine/Source/Runtime/AutoRTFM/Documentation/README.md for more info.
+)";
+}
+
+AUTORTFM_DISABLE static void ReportAutoRTFMCriticalSectionHazard()
+{
+	const autortfm_log_severity Severity = AutoRTFM::ForTheRuntime::GetAutoRTFMHazardLogSeverity();
+	LogAutoRTFMHazard(Severity, R"(%s
+A critical section was accessed while inside an AutoRTFM transaction. AutoRTFM doesn't support native
+critical sections, but it does support transactionally-safe critical sections that can be used instead.
+
+Use the call stack below to find the FCriticalSection accessed immediately before this error occurred, and
+replace it with FTransactionallySafeCriticalSection. If FScopeLock is used, replace it with UE::TScopeLock.
+)",
+		GetAutoRTFMHazardBanner());
+	AutoRTFM::LogWithCallstack(Severity, "Call stack:");
+}
+
 AUTORTFM_DISABLE static void* FunctionMapReportError(void* Open, const char* Where)
 {
+#if AUTORTFM_PLATFORM_WINDOWS
+	if (Open == &InitializeCriticalSection ||
+		Open == &InitializeCriticalSectionAndSpinCount ||
+		Open == &SetCriticalSectionSpinCount ||
+		Open == &TryEnterCriticalSection ||
+		Open == &EnterCriticalSection ||
+		Open == &LeaveCriticalSection ||
+		Open == &DeleteCriticalSection)
+	{
+		ReportAutoRTFMCriticalSectionHazard();
+	}
+#elif AUTORTFM_PLATFORM_LINUX
+	if (Open == &pthread_mutexattr_init ||
+		Open == &pthread_mutexattr_settype ||
+		Open == &pthread_mutex_init ||
+		Open == &pthread_mutexattr_destroy ||
+		Open == &pthread_mutex_destroy ||
+		Open == &pthread_mutex_trylock ||
+		Open == &pthread_mutex_lock ||
+		Open == &pthread_mutex_unlock)
+	{
+		ReportAutoRTFMCriticalSectionHazard();
+	}
+#endif
+
 	std::string FunctionDescription = GetFunctionDescription(Open);
 	if (Where)
 	{
diff --git a/Engine/Source/Runtime/AutoRTFM/Public/AutoRTFM.h b/Engine/Source/Runtime/AutoRTFM/Public/AutoRTFM.h
--- a/Engine/Source/Runtime/AutoRTFM/Public/AutoRTFM.h
+++ b/Engine/Source/Runtime/AutoRTFM/Public/AutoRTFM.h
@@ -3,10 +3,10 @@
 #pragma once
 
 /**
- *  AutoRTFM is designed to make it easy to take existing C++ code--even if it was never designed 
+ *  AutoRTFM is designed to make it easy to take existing C++ code--even if it was never designed
  *  to have any transactional semantics--and make it transactional just by using an alternate compiler.
  *  For details, see `Engine/Source/Runtime/AutoRTFM/Documentation/README.md`.
- */ 
+ */
 
 // HEADER_UNIT_SKIP - unused warnings
 
@@ -113,7 +113,7 @@ typedef struct
 	// Strings use UTF-8 encoding.
 	// Must not be null.
 	void (*LogWithCallstack)(void* ProgramCounter, autortfm_log_severity Severity, const char* Format, va_list Args);
-	
+
 	// Function used to report an ensure failure using a printf-style format string and va_list arguments.
 	// Strings use UTF-8 encoding.
 	// Must not be null.
@@ -483,11 +483,11 @@ enum class ETransactionResult
 	// cannot be undone later if the parent transaction is aborted for any reason.
     Committed = autortfm_committed,
 
-	// The transaction aborted because a new transaction nest was attempted (via OnCommit or OnComplete) while the 
+	// The transaction aborted because a new transaction nest was attempted (via OnCommit or OnComplete) while the
 	// current transaction was being committed.
 	AbortedByTransactDuringCommit = autortfm_aborted_by_transact_during_commit,
 
-	// The transaction aborted because a new transaction nest was attempted (via OnAbort or OnComplete) while the 
+	// The transaction aborted because a new transaction nest was attempted (via OnAbort or OnComplete) while the
 	// current transaction was being aborted.
 	AbortedByTransactDuringAbort = autortfm_aborted_by_transact_during_abort,
 
@@ -640,7 +640,7 @@ UE_AUTORTFM_CRITICAL_INLINE bool IsCommitting() { return autortfm_is_committing(
 UE_AUTORTFM_CRITICAL_INLINE bool IsCommittingOrAborting() { return autortfm_is_committing_or_aborting(); }
 
 // Tells us if we are currently retrying a transaction. This will return true when inside an
-// on-abort or on-complete handler that was triggered by CascadingRetryTransaction. 
+// on-abort or on-complete handler that was triggered by CascadingRetryTransaction.
 UE_AUTORTFM_CRITICAL_INLINE bool IsRetrying() { return autortfm_is_retrying(); }
 
 // Returns true if the passed-in pointer is on the stack of the currently-executing transaction.
@@ -715,17 +715,17 @@ UE_AUTORTFM_CRITICAL_INLINE void Commit(const FunctorType& Functor)
 		const_cast<void*>(static_cast<const void*>(&Functor)));
 }
 
-// Ends a transaction while in the closed, discarding all effects. 
+// Ends a transaction while in the closed, discarding all effects.
 // Sends control to the end of the transaction immediately.
 UE_AUTORTFM_CRITICAL_INLINE void AbortTransaction()
 {
 	autortfm_abort_transaction();
 }
 
-// End a transaction nest in the closed, discarding all effects. This cascades, 
+// End a transaction nest in the closed, discarding all effects. This cascades,
 // meaning an abort of a nested transaction will cause all transactions in the
 // nest to abort. Once the transaction has aborted, on-complete callbacks will
-// be invoked. Finally, control will be returned to the end of the outermost 
+// be invoked. Finally, control will be returned to the end of the outermost
 // Transact.
 #if UE_AUTORTFM
 UE_AUTORTFM_CRITICAL_INLINE void CascadingAbortTransaction()
@@ -940,7 +940,7 @@ template<typename FunctorType> [[nodiscard]] UE_AUTORTFM_CRITICAL_INLINE EContex
 }
 
 // Force a transaction nest to be retried. Once the transaction has aborted,
-// all on-complete handlers will be called before retrying the transaction. 
+// all on-complete handlers will be called before retrying the transaction.
 // This is an expensive operation and should thus be used with extreme caution.
 // If this is called outside of a transaction, the call is ignored.
 #if UE_AUTORTFM
@@ -958,7 +958,7 @@ UE_AUTORTFM_CRITICAL_INLINE void CascadingRetryTransaction()
 #endif
 
 #if UE_AUTORTFM
-// Have some work happen when this transaction commits. 
+// Have some work happen when this transaction commits.
 // In a nested transaction, the work is deferred until the outermost nest is committed;
 // at that point, the worklist is run in FIFO order.
 // If this is called outside a transaction or from an open nest, then the work
@@ -1229,7 +1229,7 @@ namespace ForTheRuntime
 	// An enum to represent the various ways we want to enable/disable the AutoRTFM runtime.
 	// This enum has effective groups of functionality such that if a higher priority group
 	// has enabled or disabled the runtime, a lower priority group cannot then override that.
-	// 
+	//
 	// We have from higher to lower priority:
 	// - Forced Enabled/Disabled - used by CVars when force enabling/disabling AutoRTFM.
 	// - Override Enabled/Disabled - override any setting of enabled/disabled as was set by a CVar.
@@ -1326,10 +1326,13 @@ namespace ForTheRuntime
 	UE_AUTORTFM_API void SetInternalAbortAction(EAutoRTFMInternalAbortActionState State);
 
 	UE_AUTORTFM_API EAutoRTFMInternalAbortActionState GetInternalAbortAction();
-	
+
 	UE_AUTORTFM_API bool GetEnsureOnInternalAbort();
 	UE_AUTORTFM_API void SetEnsureOnInternalAbort(bool bEnabled);
 
+	UE_AUTORTFM_API autortfm_log_severity GetAutoRTFMHazardLogSeverity();
+	UE_AUTORTFM_API void SetAutoRTFMHazardLogSeverity(autortfm_log_severity Severity);
+
 	// Set whether we should trigger an ensure on an abort-by-language.
 	[[deprecated("Use `SetEnsureOnInternalAbort` instead!")]]
 	inline void SetEnsureOnAbortByLanguage(bool bEnabled)
@@ -1391,7 +1394,7 @@ namespace ForTheRuntime
 	// Can only be called within an already active parent transaction (EG. this cannot start
 	// a transaction nest itself).
 	AUTORTFM_DISABLE UE_AUTORTFM_API void StartTransaction();
-	
+
 	// Manually commit the top transaction nest, popping it from the execution scope.
 	// Can only be called within an already active parent transaction (EG. this cannot end
 	// a transaction nest itself).
@@ -1411,7 +1414,7 @@ namespace ForTheRuntime
 	// RollbackTransaction() sets the transaction status to AbortedByRequest.
 	AUTORTFM_DISABLE UE_AUTORTFM_API void RollbackTransaction();
 
-	// CascadingAbortRollbackTransaction() rolls back the current scoped 
+	// CascadingAbortRollbackTransaction() rolls back the current scoped
 	// transaction, which must be a scoped transaction (i.e. created by
 	// AutoRTFM::Transact() and not by StartTransaction()).
 	// CascadingAbortRollbackTransaction() sets the transaction status to
@@ -1421,7 +1424,7 @@ namespace ForTheRuntime
 	// undefined behaviour.
 	AUTORTFM_DISABLE UE_AUTORTFM_API void CascadingAbortRollbackTransaction();
 
-	// CascadingRetryRollbackTransaction() rolls back the current scoped 
+	// CascadingRetryRollbackTransaction() rolls back the current scoped
 	// transaction, which must be a scoped transaction (i.e. created by
 	// AutoRTFM::Transact() and not by StartTransaction()).
 	// CascadingRetryRollbackTransaction() sets the transaction status to
@@ -1540,8 +1543,8 @@ namespace AutoRTFM::Private
 #endif
 
 // Declares an AutoRTFM-aware thread local variable. `thread_local` variables are not yet natively supported (#jira SOL-7684)
-// Calls should be written like this: 
-//     UE_AUTORTFM_DECLARE_THREAD_LOCAL_VAR(FString, MyThreadLocalString); 
+// Calls should be written like this:
+//     UE_AUTORTFM_DECLARE_THREAD_LOCAL_VAR(FString, MyThreadLocalString);
 //     MyThreadLocalString = "Hello";
 #define UE_AUTORTFM_DECLARE_THREAD_LOCAL_VAR(Type, Name) UE_AUTORTFM_DECLARE_THREAD_LOCAL_VAR_IMPL(Type, Name)
 
diff --git a/Engine/Source/Runtime/Core/Private/AutoRTFM/AutoRTFMUE.cpp b/Engine/Source/Runtime/Core/Private/AutoRTFM/AutoRTFMUE.cpp
--- a/Engine/Source/Runtime/Core/Private/AutoRTFM/AutoRTFMUE.cpp
+++ b/Engine/Source/Runtime/Core/Private/AutoRTFM/AutoRTFMUE.cpp
@@ -11,7 +11,6 @@
 #include "HAL/IConsoleManager.h"
 #include "HAL/LowLevelMemTracker.h"
 #include "HAL/UnrealMemory.h"
-#include "Logging/LogMacros.h"
 #include "Logging/LogVerbosity.h"
 #include "Misc/AssertionMacros.h"
 
@@ -19,8 +18,6 @@
 #include <cstdio>
 #include <memory>
 
-DECLARE_LOG_CATEGORY_EXTERN(LogAutoRTFM, Display, All)
-
 DEFINE_LOG_CATEGORY(LogAutoRTFM)
 
 static_assert(UE_AUTORTFM_ENABLED, "AutoRTFM/AutoRTFMUE.cpp requires the compiler flag '-fautortfm'");
diff --git a/Engine/Source/Runtime/Core/Public/AutoRTFM/AutoRTFMUE.h b/Engine/Source/Runtime/Core/Public/AutoRTFM/AutoRTFMUE.h
--- a/Engine/Source/Runtime/Core/Public/AutoRTFM/AutoRTFMUE.h
+++ b/Engine/Source/Runtime/Core/Public/AutoRTFM/AutoRTFMUE.h
@@ -6,6 +6,9 @@
 #include "HAL/Platform.h"
 
 #include "AutoRTFM.h"
+#include "Logging/LogMacros.h"
+
+UE_AUTORTFM_API DECLARE_LOG_CATEGORY_EXTERN(LogAutoRTFM, Display, All)
 
 namespace AutoRTFM
 {
