[AutoRTFM] Untangle control flow in ReportError.

This code should behave the same as before, but the flow is simpler to reason about (for me, at least).

#FailedValidations Preflight

diff --git a/Engine/Source/Runtime/AutoRTFM/Private/Utils.cpp b/Engine/Source/Runtime/AutoRTFM/Private/Utils.cpp
--- a/Engine/Source/Runtime/AutoRTFM/Private/Utils.cpp
+++ b/Engine/Source/Runtime/AutoRTFM/Private/Utils.cpp
@@ -69,57 +69,46 @@ UE_AUTORTFM_API UE_AUTORTFM_FORCENOINLINE void ReportError(const char* File, int
 
 	const ForTheRuntime::EAutoRTFMInternalAbortActionState InternalAbortAction = ForTheRuntime::GetInternalAbortAction();
 
-	if (Format)
+	switch (InternalAbortAction)
 	{
-		if (InternalAbortAction == ForTheRuntime::EAutoRTFMInternalAbortActionState::Crash)
-		{
-			va_list Args;
-			va_start(Args, Format);
-			::AutoRTFM::LogV(File, Line, ProgramCounter, autortfm_log_fatal, Format, Args);
-			va_end(Args);
-		}
-		else if (ForTheRuntime::GetEnsureOnInternalAbort())
-		{
-			va_list Args;
-			va_start(Args, Format);
-			[[maybe_unused]] static bool bCalled = [&]
+		case ForTheRuntime::EAutoRTFMInternalAbortActionState::Crash:
+			if (Format)
 			{
-				::AutoRTFM::EnsureFailureV(File, Line, ProgramCounter, "!GetEnsureOnInternalAbort()", Format, Args);
-				return true;
-			}();
-			va_end(Args);
-		}
-	}
-	else
-	{
-		if (InternalAbortAction == ForTheRuntime::EAutoRTFMInternalAbortActionState::Crash)
-		{
-			::AutoRTFM::Log(File, Line, ProgramCounter, autortfm_log_fatal, "Transaction failing because of internal issue");
-		}
-		else if (ForTheRuntime::GetEnsureOnInternalAbort())
-		{
-			[[maybe_unused]] static bool bCalled = [&]
+				va_list Args;
+				va_start(Args, Format);
+				::AutoRTFM::LogV(File, Line, ProgramCounter, autortfm_log_fatal, Format, Args);
+				va_end(Args);
+			}
+			else
 			{
-				::AutoRTFM::EnsureFailure(File, Line, ProgramCounter, "!GetEnsureOnInternalAbort()", "Transaction failing because of internal issue");
-				return true;
-			}();
-		}
-	}
+				::AutoRTFM::Log(File, Line, ProgramCounter, autortfm_log_fatal, "Transaction failing because of internal issue");
+			}
+			// The `autortfm_log_fatal` lines above probably triggered a crash already, but if we made it this far, force the crash.
+			__builtin_trap();
+			break;
 
-	switch (InternalAbortAction)
-	{
 		case ForTheRuntime::EAutoRTFMInternalAbortActionState::Abort:
+			if (ForTheRuntime::GetEnsureOnInternalAbort())
+			{
+				if (Format == nullptr)
+				{
+					Format = "Transaction failing because of internal issue";
+				}
+				va_list Args;
+				va_start(Args, Format);
+				[[maybe_unused]] static bool bCalled = [&]
+				{
+					::AutoRTFM::EnsureFailureV(File, Line, ProgramCounter, "!GetEnsureOnInternalAbort()", Format, Args);
+					return true;
+				}();
+				va_end(Args);
+			}
 			if (FContext* const Context = FContext::Get())
 			{
 				Context->AbortByLanguageAndThrow();
 			}
 			break;
 
-		case ForTheRuntime::EAutoRTFMInternalAbortActionState::Crash:
-			// The `autortfm_log_fatal` lines above probably triggered a crash already, but if we made it this far, force the crash.
-			__builtin_trap();
-			break;
-
 		default:
 			::AutoRTFM::EnsureFailure(File, Line, ProgramCounter, "Unexpected InternalAbortAction", "InternalAbortAction = %d", int(InternalAbortAction));
 			break;
