[AutoRTFM] Fix issue where inlined functions would generate unnecessary closed lambda variants.

We were only considering 'internal' linkage for eliding closed variants. Inlined functions and their lambdas have 'linkonce_odr' linkage, and so were not getting omitted even when all callsites were AutoRTFM-disabled. To fix, also consider 'linkonce_odr' linkage functions.

Also remove the pointless closed function lookup for `AutoRTFM::TransactThenOpen()` - the closed variant lookup is not used, and breaks the closed elision for lambdas handed to this function.

diff --git a/Engine/Source/Programs/AutoRTFMTests/Private/OpenAPI.cpp b/Engine/Source/Programs/AutoRTFMTests/Private/OpenAPI.cpp
--- a/Engine/Source/Programs/AutoRTFMTests/Private/OpenAPI.cpp
+++ b/Engine/Source/Programs/AutoRTFMTests/Private/OpenAPI.cpp
@@ -982,8 +982,6 @@ TEST_CASE("OpenAPI.DeferredStartTransactionOverflow")
 
 	AutoRTFM::TransactThenOpen([&]
 	{
-		AutoRTFM::UnreachableIfClosed("#jira SOL-8290");
-		
 		uint64 NestingCount = static_cast<uint64>(std::numeric_limits<uint16>::max()) + 42;
 		for (uint64 I = 0; I < NestingCount; ++I)
 		{
diff --git a/Engine/Source/Runtime/AutoRTFM/Private/API.cpp b/Engine/Source/Runtime/AutoRTFM/Private/API.cpp
--- a/Engine/Source/Runtime/AutoRTFM/Private/API.cpp
+++ b/Engine/Source/Runtime/AutoRTFM/Private/API.cpp
@@ -473,7 +473,7 @@ UE_AUTORTFM_FORCEINLINE autortfm_result TransactThenOpenImpl(void (*Uninstrument
 		}));
 }
 
-AUTORTFM_INTERNAL autortfm_result RTFM_autortfm_transact_then_open(void (*UninstrumentedWork)(void*), void (*InstrumentedWork)(void*), void* Arg, const void* ReturnAddress)
+AUTORTFM_INTERNAL autortfm_result RTFM_autortfm_transact_then_open(void (*UninstrumentedWork)(void*), void* Arg, const void* ReturnAddress)
 {
 	return TransactThenOpenImpl(UninstrumentedWork, Arg, ReturnAddress);
 }
@@ -847,11 +847,11 @@ extern "C" autortfm_result autortfm_transact(void (*UninstrumentedWork)(void*),
 	return autortfm_committed;
 }
 
-extern "C" autortfm_result autortfm_transact_then_open(void (*UninstrumentedWork)(void*), void (*InstrumentedWork)(void*), void* Arg)
+extern "C" autortfm_result autortfm_transact_then_open(void (*UninstrumentedWork)(void*), void* Arg)
 {
 	if (autortfm_is_closed())
 	{
-		return RTFM_autortfm_transact_then_open(UninstrumentedWork, InstrumentedWork, Arg, __builtin_return_address(0));
+		return RTFM_autortfm_transact_then_open(UninstrumentedWork, Arg, __builtin_return_address(0));
 	}
 
 	return TransactThenOpenImpl(UninstrumentedWork, Arg, __builtin_return_address(0));
diff --git a/Engine/Source/Runtime/AutoRTFM/Public/AutoRTFM.h b/Engine/Source/Runtime/AutoRTFM/Public/AutoRTFM.h
--- a/Engine/Source/Runtime/AutoRTFM/Public/AutoRTFM.h
+++ b/Engine/Source/Runtime/AutoRTFM/Public/AutoRTFM.h
@@ -222,11 +222,10 @@ UE_AUTORTFM_CRITICAL_INLINE autortfm_result autortfm_transact(void (*Uninstrumen
 #endif
 
 #if UE_AUTORTFM_ENABLED
-UE_AUTORTFM_API autortfm_result autortfm_transact_then_open(void (*UninstrumentedWork)(void*), void (*InstrumentedWork)(void*), void* Arg);
+UE_AUTORTFM_API autortfm_result autortfm_transact_then_open(void (*UninstrumentedWork)(void*), void* Arg) AUTORTFM_EXCEPT;
 #else
-UE_AUTORTFM_CRITICAL_INLINE autortfm_result autortfm_transact_then_open(void (*UninstrumentedWork)(void*), void (*InstrumentedWork)(void*), void* Arg)
+UE_AUTORTFM_CRITICAL_INLINE autortfm_result autortfm_transact_then_open(void (*UninstrumentedWork)(void*), void* Arg) AUTORTFM_EXCEPT
 {
-	UE_AUTORTFM_UNUSED(InstrumentedWork);
 	UninstrumentedWork(Arg);
     return autortfm_committed;
 }
@@ -693,7 +692,6 @@ UE_AUTORTFM_CRITICAL_INLINE ETransactionResult TransactThenOpen(const FunctorTyp
 		static_cast<ETransactionResult>(
 			autortfm_transact_then_open(
 				&AutoRTFMFunctorInvoker<FunctorType>,
-				AutoRTFMLookupInstrumentedFunctorInvoker<FunctorType>(Functor),
 				const_cast<void*>(static_cast<const void*>(&Functor))));
 
 	return Result;
diff --git a/Engine/Source/Runtime/Core/Public/Stats/DynamicStats.h b/Engine/Source/Runtime/Core/Public/Stats/DynamicStats.h
--- a/Engine/Source/Runtime/Core/Public/Stats/DynamicStats.h
+++ b/Engine/Source/Runtime/Core/Public/Stats/DynamicStats.h
@@ -87,9 +87,6 @@ private: // private since this can only be declared if STATS is defined, due to
 	{
 		return AutoRTFM::Open([&]
 		{
-			// #jira SOL-8290: can remove this `UnreachableIfClosed` once resolved.
-			AutoRTFM::UnreachableIfClosed();
-			
 			return CreateStatIdInternalInTheOpen<TStatGroup>(StatNameOrDescription, Type, IsTimer, bClearEveryFrame);
 		});
 	}
@@ -117,9 +114,6 @@ private: // private since this can only be declared if STATS is defined, due to
 	{
 		return AutoRTFM::Open([&]
 		{
-			// #jira SOL-8290: can remove this `UnreachableIfClosed` once resolved.
-			AutoRTFM::UnreachableIfClosed();
-
 			return CreateMemoryStatIdInternalInTheOpen<TStatGroup>(StatNameOrDescription, MemRegion);
 		});
 	}
