[AutoRTFM] Fix log spam when running AutoRTFMTests.

Warnings in particular can be captured and silenced at test time by using FCaptureWarningContext.

This CL also migrates the hazard logging code to a separate CPP since it might be needed outside of FunctionMap.

#FailedValidations Preflight

diff --git a/Engine/Source/Programs/AutoRTFMTests/Private/Abort.cpp b/Engine/Source/Programs/AutoRTFMTests/Private/Abort.cpp
--- a/Engine/Source/Programs/AutoRTFMTests/Private/Abort.cpp
+++ b/Engine/Source/Programs/AutoRTFMTests/Private/Abort.cpp
@@ -1364,31 +1364,17 @@ namespace
 
 class FScopedErrorHintTest
 {
-	public:
-		FScopedErrorHintTest(FOutputDevice& InDevice) : Device(InDevice)
-		{
-			GLog->AddOutputDevice(&Device);
-		}
-
-		~FScopedErrorHintTest()
-		{
-			GLog->RemoveOutputDevice(&Device);
-		}
-
 	private:
 		AutoRTFMTestUtils::FScopedInternalAbortAction Scoped1{AutoRTFM::ForTheRuntime::EAutoRTFMInternalAbortActionState::Abort};
 		AutoRTFMTestUtils::FScopedEnsureOnInternalAbort Scoped2{false};
-		AutoRTFMTestUtils::FScopedHazardLogSeverity Scoped3{autortfm_log_info};
-		FLogScopedVerbosityOverride Scoped4{&LogAutoRTFM, ELogVerbosity::Display};
-		FOutputDevice& Device;
 };
 
 }
 
 TEST_CASE("Abort.Language.ErrorHints.InitializeCriticalSection")
 {
-	TLogCategoryOutputDevice StringLog(LogAutoRTFM);
-	FScopedErrorHintTest TestWith(StringLog);
+	FScopedErrorHintTest ScopedTest;
+    AutoRTFMTestUtils::FCaptureWarningContext WarningContext;
 
 	const AutoRTFM::ETransactionResult Result = AutoRTFM::Transact([&]
 	{
@@ -1398,15 +1384,15 @@ TEST_CASE("Abort.Language.ErrorHints.InitializeCriticalSection")
 	REQUIRE(AutoRTFM::ETransactionResult::AbortedByLanguage == Result);
 
 	// Did we print our detailed error message?
-	REQUIRE(StringLog.Contains("A critical section was accessed while inside an AutoRTFM transaction."));
-	REQUIRE(StringLog.Contains("AutoRTFM/Documentation/README.md"));
-	REQUIRE(StringLog.Contains("Use the call stack below to find the FCriticalSection"));
+	REQUIRE(WarningContext.HasWarningSubstring("A critical section was accessed while inside an AutoRTFM transaction."));
+	REQUIRE(WarningContext.HasWarningSubstring("AutoRTFM/Documentation/README.md"));
+	REQUIRE(WarningContext.HasWarningSubstring("Use the call stack below to find the FCriticalSection"));
 }
 
 TEST_CASE("Abort.Language.ErrorHints.EnterCriticalSection")
 {
-	TLogCategoryOutputDevice StringLog(LogAutoRTFM);
-	FScopedErrorHintTest TestWith(StringLog);
+	FScopedErrorHintTest ScopedTest;
+    AutoRTFMTestUtils::FCaptureWarningContext WarningContext;
 
 	FCriticalSection CS;
 
@@ -1418,13 +1404,13 @@ TEST_CASE("Abort.Language.ErrorHints.EnterCriticalSection")
 	REQUIRE(AutoRTFM::ETransactionResult::AbortedByLanguage == Result);
 
 	// Did we print our message?
-	REQUIRE(StringLog.Contains("A critical section was accessed while inside an AutoRTFM transaction."));
+	REQUIRE(WarningContext.HasWarningSubstring("A critical section was accessed while inside an AutoRTFM transaction."));
 }
 
 TEST_CASE("Abort.Language.ErrorHints.AcquireSRWLockExclusive")
 {
-	TLogCategoryOutputDevice StringLog(LogAutoRTFM);
-	FScopedErrorHintTest TestWith(StringLog);
+	FScopedErrorHintTest ScopedTest;
+    AutoRTFMTestUtils::FCaptureWarningContext WarningContext;
 
 	FRWLock RW;
 
@@ -1436,13 +1422,13 @@ TEST_CASE("Abort.Language.ErrorHints.AcquireSRWLockExclusive")
 	REQUIRE(AutoRTFM::ETransactionResult::AbortedByLanguage == Result);
 
 	// Did we print our message?
-	REQUIRE(StringLog.Contains("A read-write lock was accessed while inside an AutoRTFM transaction."));
+	REQUIRE(WarningContext.HasWarningSubstring("A read-write lock was accessed while inside an AutoRTFM transaction."));
 }
 
 TEST_CASE("Abort.Language.ErrorHints.TryAcquireSRWLockShared")
 {
-	TLogCategoryOutputDevice StringLog(LogAutoRTFM);
-	FScopedErrorHintTest TestWith(StringLog);
+	FScopedErrorHintTest ScopedTest;
+    AutoRTFMTestUtils::FCaptureWarningContext WarningContext;
 
 	FRWLock RW;
 
@@ -1454,5 +1440,5 @@ TEST_CASE("Abort.Language.ErrorHints.TryAcquireSRWLockShared")
 	REQUIRE(AutoRTFM::ETransactionResult::AbortedByLanguage == Result);
 
 	// Did we print our message?
-	REQUIRE(StringLog.Contains("A read-write lock was accessed while inside an AutoRTFM transaction."));
+	REQUIRE(WarningContext.HasWarningSubstring("A read-write lock was accessed while inside an AutoRTFM transaction."));
 }
diff --git a/Engine/Source/Programs/AutoRTFMTests/Private/AutoRTFMTestUtils.h b/Engine/Source/Programs/AutoRTFMTests/Private/AutoRTFMTestUtils.h
--- a/Engine/Source/Programs/AutoRTFMTests/Private/AutoRTFMTestUtils.h
+++ b/Engine/Source/Programs/AutoRTFMTests/Private/AutoRTFMTestUtils.h
@@ -184,25 +184,4 @@ private:
 	const AutoRTFM::ForTheRuntime::EAutoRTFMInternalAbortActionState Original;
 };
 
-class FScopedHazardLogSeverity final
-{
-public:
-	FScopedHazardLogSeverity(autortfm_log_severity Severity)
-		: Original(AutoRTFM::ForTheRuntime::GetAutoRTFMHazardLogSeverity())
-	{
-		AutoRTFM::ForTheRuntime::SetAutoRTFMHazardLogSeverity(Severity);
-	}
-
-	~FScopedHazardLogSeverity()
-	{
-		AutoRTFM::ForTheRuntime::SetAutoRTFMHazardLogSeverity(Original);
-	}
-
-private:
-	FScopedHazardLogSeverity(const FScopedHazardLogSeverity&) = delete;
-	FScopedHazardLogSeverity& operator=(const FScopedHazardLogSeverity&) = delete;
-
-	const autortfm_log_severity Original;
-};
-
 }
diff --git a/Engine/Source/Runtime/AutoRTFM/Private/API.cpp b/Engine/Source/Runtime/AutoRTFM/Private/API.cpp
--- a/Engine/Source/Runtime/AutoRTFM/Private/API.cpp
+++ b/Engine/Source/Runtime/AutoRTFM/Private/API.cpp
@@ -40,7 +40,6 @@ namespace
 	bool GMemoryValidationThrottlingEnabled = true;
 	bool GMemoryValidationStatisticsEnabled = true;
 	bool GAutoRTFMEnsureOnInternalAbort = true;
-	autortfm_log_severity GAutoRTFMHazardLogSeverity = autortfm_log_error;
 
 	// A linked-list of open->closed function tables populated by
 	// autortfm_register_open_to_closed_functions(). This is consumed by
@@ -263,24 +262,6 @@ namespace AutoRTFM
 #endif
 		}
 
-		autortfm_log_severity GetAutoRTFMHazardLogSeverity()
-		{
-#if UE_AUTORTFM
-			return GAutoRTFMHazardLogSeverity;
-#else
-			return autortfm_log_error;
-#endif
-		}
-
-		void SetAutoRTFMHazardLogSeverity(autortfm_log_severity Severity)
-		{
-#if UE_AUTORTFM
-			GAutoRTFMHazardLogSeverity = Severity;
-#else
-			UE_AUTORTFM_UNUSED(Severity);
-#endif
-		}
-
 		void SetRetryTransaction(EAutoRTFMRetryTransactionState State)
 		{
 #if UE_AUTORTFM
diff --git a/Engine/Source/Runtime/AutoRTFM/Private/FunctionMap.cpp b/Engine/Source/Runtime/AutoRTFM/Private/FunctionMap.cpp
--- a/Engine/Source/Runtime/AutoRTFM/Private/FunctionMap.cpp
+++ b/Engine/Source/Runtime/AutoRTFM/Private/FunctionMap.cpp
@@ -7,6 +7,7 @@
 #include "AutoRTFM.h"
 #include "BuildMacros.h"
 #include "FunctionMapInlines.h"
+#include "ReportAutoRTFMHazard.h"
 #include "Utils.h"
 
 #include <cinttypes>
@@ -257,54 +258,6 @@ inline void* FollowRelocation(void* Function)
 	return nullptr;
 }
 
-template <typename... ArgList>
-AUTORTFM_DISABLE
-static void LogAutoRTFMHazard(autortfm_log_severity Severity, const char* Message, ArgList&&... Args)
-{
-	AutoRTFM::Log(__FILE__, __LINE__, /* ProgramCounter */ nullptr, Severity, Message, Args...);
-}
-
-AUTORTFM_DISABLE static const char* GetAutoRTFMHazardBanner()
-{
-	return R"(
-             _         _        ____ _____ _____ __  __   _   _                        _
-            / \  _   _| |_ ___ |  _ \_   _|  ___|  \/  | | | | | __ _ ______ _ _ __ __| |
-  _____    / _ \| | | | __/ _ \| |_) || | | |_  | |\/| | | |_| |/ _` |_  / _` | '__/ _` |  _____
- |_____|  / ___ \ |_| | || (_) |  _ < | | |  _| | |  | | |  _  | (_| |/ / (_| | | | (_| | |_____|
-         /_/   \_\__,_|\__\___/|_| \_\|_| |_|   |_|  |_| |_| |_|\__,_/___\__,_|_|  \__,_|
-
-New to AutoRTFM? See Engine/Source/Runtime/AutoRTFM/Documentation/README.md for more info.
-)";
-}
-
-AUTORTFM_DISABLE static void ReportAutoRTFMCriticalSectionHazard()
-{
-	const autortfm_log_severity Severity = AutoRTFM::ForTheRuntime::GetAutoRTFMHazardLogSeverity();
-	LogAutoRTFMHazard(Severity, R"(%s
-A critical section was accessed while inside an AutoRTFM transaction. AutoRTFM doesn't support native
-critical sections, but it does support transactionally-safe critical sections that can be used instead.
-
-Use the call stack below to find the FCriticalSection accessed immediately before this error occurred, and
-replace it with FTransactionallySafeCriticalSection. If FScopeLock is used, replace it with UE::TScopeLock.
-)",
-		GetAutoRTFMHazardBanner());
-	AutoRTFM::LogWithCallstack(Severity, "Call stack:");
-}
-
-AUTORTFM_DISABLE static void ReportAutoRTFMRWLockHazard()
-{
-	const autortfm_log_severity Severity = AutoRTFM::ForTheRuntime::GetAutoRTFMHazardLogSeverity();
-	LogAutoRTFMHazard(Severity, R"(%s
-A read-write lock was accessed while inside an AutoRTFM transaction. AutoRTFM doesn't support native
-read-write locks, but it does support transactionally-safe read-write locks that can be used instead.
-
-Use the call stack below to find the FRWLock accessed immediately before this error occurred, and
-replace it with FTransactionallySafeFRWLock. If FReadScopeLock, FWriteScopeLock or FRWScopeLock are
-used, replace them with UE::TReadScopeLock, UE::TWriteScopeLock or UE::TRWScopeLock.
-)",
-		GetAutoRTFMHazardBanner());
-	AutoRTFM::LogWithCallstack(Severity, "Call stack:");
-}
 
 AUTORTFM_DISABLE static void* FunctionMapReportError(void* Open, const char* Where)
 {
@@ -317,7 +270,7 @@ AUTORTFM_DISABLE static void* FunctionMapReportError(void* Open, const char* Whe
 		Open == &LeaveCriticalSection ||
 		Open == &DeleteCriticalSection)
 	{
-		ReportAutoRTFMCriticalSectionHazard();
+		ForTheRuntime::ReportAutoRTFMHazard(ForTheRuntime::EHazardType::CriticalSection);
 	}
 	else if (Open == &TryAcquireSRWLockExclusive ||
 		Open == &AcquireSRWLockExclusive ||
@@ -327,7 +280,7 @@ AUTORTFM_DISABLE static void* FunctionMapReportError(void* Open, const char* Whe
 		Open == &ReleaseSRWLockShared ||
 		Open == &TryAcquireSRWLockExclusive)
 	{
-		ReportAutoRTFMRWLockHazard();
+		ForTheRuntime::ReportAutoRTFMHazard(ForTheRuntime::EHazardType::RWLock);
 	}
 #elif AUTORTFM_PLATFORM_LINUX
 	if (Open == &pthread_mutexattr_init ||
@@ -339,7 +292,7 @@ AUTORTFM_DISABLE static void* FunctionMapReportError(void* Open, const char* Whe
 		Open == &pthread_mutex_lock ||
 		Open == &pthread_mutex_unlock)
 	{
-		ReportAutoRTFMCriticalSectionHazard();
+		ForTheRuntime::ReportAutoRTFMHazard(ForTheRuntime::EHazardType::CriticalSection);
 	}
 	else if (Open == &pthread_rwlock_init ||
 		Open == &pthread_rwlock_destroy ||
@@ -349,7 +302,7 @@ AUTORTFM_DISABLE static void* FunctionMapReportError(void* Open, const char* Whe
 		Open == &pthread_rwlock_tryrdlock ||
 		Open == &pthread_rwlock_rdlock)
 	{
-		ReportAutoRTFMRWLockHazard();
+		ForTheRuntime::ReportAutoRTFMHazard(ForTheRuntime::EHazardType::RWLock);
 	}
 #endif
 
diff --git a/Engine/Source/Runtime/AutoRTFM/Private/ReportAutoRTFMHazard.cpp b/Engine/Source/Runtime/AutoRTFM/Private/ReportAutoRTFMHazard.cpp
new file mode 100644
--- /dev/null
+++ b/Engine/Source/Runtime/AutoRTFM/Private/ReportAutoRTFMHazard.cpp
@@ -0,0 +1,54 @@
+// Copyright Epic Games, Inc. All Rights Reserved.
+
+#if (defined(__AUTORTFM) && __AUTORTFM)
+
+#include "ReportAutoRTFMHazard.h"
+#include "Utils.h"
+
+void AutoRTFM::ForTheRuntime::ReportAutoRTFMHazard(EHazardType HazardType)
+{
+	static constexpr const char Banner[] = R"(
+             _         _        ____ _____ _____ __  __   _   _                        _
+            / \  _   _| |_ ___ |  _ \_   _|  ___|  \/  | | | | | __ _ ______ _ _ __ __| |
+  _____    / _ \| | | | __/ _ \| |_) || | | |_  | |\/| | | |_| |/ _` |_  / _` | '__/ _` |  _____
+ |_____|  / ___ \ |_| | || (_) |  _ < | | |  _| | |  | | |  _  | (_| |/ / (_| | | | (_| | |_____|
+         /_/   \_\__,_|\__\___/|_| \_\|_| |_|   |_|  |_| |_| |_|\__,_/___\__,_|_|  \__,_|
+
+New to AutoRTFM? See Engine/Source/Runtime/AutoRTFM/Documentation/README.md for more info.
+)";
+
+	const char* Message = "An unknown error occurred.";
+
+	switch (HazardType)
+	{
+		case EHazardType::CriticalSection:
+			Message = R"(
+A critical section was accessed while inside an AutoRTFM transaction. AutoRTFM doesn't support native
+critical sections, but it does support transactionally-safe critical sections that can be used instead.
+
+Use the call stack below to find the FCriticalSection accessed immediately before this error occurred, and
+replace it with FTransactionallySafeCriticalSection. If FScopeLock is used, replace it with UE::TScopeLock.
+)";
+			break;
+
+		case EHazardType::RWLock:
+			Message = R"(
+A read-write lock was accessed while inside an AutoRTFM transaction. AutoRTFM doesn't support native
+read-write locks, but it does support transactionally-safe read-write locks that can be used instead.
+
+Use the call stack below to find the FRWLock accessed immediately before this error occurred, and
+replace it with FTransactionallySafeFRWLock. If FReadScopeLock, FWriteScopeLock or FRWScopeLock are
+used, replace them with UE::TReadScopeLock, UE::TWriteScopeLock or UE::TRWScopeLock.
+)";
+			break;
+
+		default:
+			AUTORTFM_ENSURE_MSG(false, "Unrecognized hazard type %d", (int)HazardType);
+			break;
+	}
+
+	AUTORTFM_WARN("%s%s", Banner, Message);
+	::AutoRTFM::LogWithCallstack(autortfm_log_warn, "Call stack:");
+}
+
+#endif // defined(__AUTORTFM) && __AUTORTFM
diff --git a/Engine/Source/Runtime/AutoRTFM/Private/ReportAutoRTFMHazard.h b/Engine/Source/Runtime/AutoRTFM/Private/ReportAutoRTFMHazard.h
new file mode 100644
--- /dev/null
+++ b/Engine/Source/Runtime/AutoRTFM/Private/ReportAutoRTFMHazard.h
@@ -0,0 +1,22 @@
+// Copyright Epic Games, Inc. All Rights Reserved.
+
+#pragma once
+
+#include "AutoRTFM.h"
+
+#if UE_AUTORTFM
+
+namespace AutoRTFM::ForTheRuntime
+{
+
+enum class EHazardType
+{
+	CriticalSection,
+	RWLock,
+};
+
+AUTORTFM_DISABLE void ReportAutoRTFMHazard(EHazardType HazardType);
+
+} // namespace AutoRTFM::ForTheRuntime
+
+#endif // UE_AUTORTFMS
\ No newline at end of file
