Fix BugHawk-identified typo.

diff --git a/Engine/Source/Runtime/AutoRTFM/Private/ReportAutoRTFMHazard.h b/Engine/Source/Runtime/AutoRTFM/Private/ReportAutoRTFMHazard.h
--- a/Engine/Source/Runtime/AutoRTFM/Private/ReportAutoRTFMHazard.h
+++ b/Engine/Source/Runtime/AutoRTFM/Private/ReportAutoRTFMHazard.h
@@ -19,4 +19,4 @@ AUTORTFM_DISABLE void ReportAutoRTFMHazard(EHazardType HazardType);
 
 } // namespace AutoRTFM::ForTheRuntime
 
-#endif // UE_AUTORTFMS
\ No newline at end of file
+#endif // UE_AUTORTFM
\ No newline at end of file
