[AutoRTFM] Move C APIs into a separate header.

This is an incremental step towards non-monolithic AutoRTFM header support. The C API header is extremely simple, containing mostly simple structs, simple enums, prototypes (when AutoRTFM exists) or empty functions (when AutoRTFM does not exist).

diff --git a/Engine/Source/Runtime/AutoRTFM/Public/AutoRTFM.h b/Engine/Source/Runtime/AutoRTFM/Public/AutoRTFM.h
--- a/Engine/Source/Runtime/AutoRTFM/Public/AutoRTFM.h
+++ b/Engine/Source/Runtime/AutoRTFM/Public/AutoRTFM.h
@@ -12,441 +12,19 @@
 
 #include "AutoRTFMConstants.h"
 #include "AutoRTFMDefines.h" // IWYU pragma: export
+#include "AutoRTFM/CAPI.h"
 
 #ifdef __cplusplus
+
 #include "AutoRTFMTask.h"
 
 #include <cstdarg>
 #include <new>
 #include <type_traits>
 #include <utility>
-#endif // __cplusplus
-
 #include <memory.h>
 #include <stdlib.h>
 
-#ifdef __cplusplus
-extern "C"
-{
-#endif
-
-// The C API exists for a few reasons:
-//
-// - It makes linking easy. AutoRTFM has to deal with a weird kind of linking
-//   where the compiler directly emits calls to functions with a given name.
-//   It's easiest to do that in llvm if the functions have C linkage and C ABI.
-// - It makes testing easy. Even seemingly simple C++ code introduces pitfalls
-//   for AutoRTFM. So very focused tests work best when written in C.
-// - It makes compiler optimizations much easier as there is no mangling to
-// 	 consider when looking for functions in the runtime we can optimize.
-//
-// We use snake_case for C API surface area to make it easy to distinguish.
-//
-// The C API should not be used directly - it is here purely as an
-// implementation detail.
-
-// This must match AutoRTFM::ETransactionResult.
-typedef enum
-{
-    autortfm_aborted_by_request = 0,
-    autortfm_aborted_by_language,
-    autortfm_committed,
-	autortfm_aborted_by_transact_during_commit,
-	autortfm_aborted_by_transact_during_abort,
-	autortfm_aborted_by_cascade,
-} autortfm_result;
-
-// This must match AutoRTFM::EContextStatus.
-typedef enum
-{
-	autortfm_status_idle = 0,
-	autortfm_status_ontrack,
-	autortfm_status_aborted_by_failed_lock_aquisition,
-	autortfm_status_aborted_by_language,
-	autortfm_status_aborted_by_request,
-	autortfm_status_committing,
-	autortfm_status_aborted_by_cascading_abort,
-	autortfm_status_aborted_by_cascading_retry,
-	autortfm_status_in_static_local_initializer,
-	autortfm_status_in_post_abort
-} autortfm_status;
-
-// AutoRTFM logging severity.
-typedef enum
-{
-	autortfm_log_verbose = 0,
-	autortfm_log_info,
-	autortfm_log_warn,
-	autortfm_log_error,
-	autortfm_log_fatal,
-} autortfm_log_severity;
-
-// An opaque unique identifier for a transaction.
-typedef uint64_t autortfm_transaction_id;
-
-// Function pointers used by AutoRTFM for heap allocations, etc.
-typedef struct
-{
-	// The function used to allocate memory from the heap.
-	// Must not be null.
-	void* (*Allocate)(size_t Size, size_t Alignment);
-
-	// The function used to reallocate memory from the heap.
-	// Must not be null.
-	void* (*Reallocate)(void* Pointer, size_t Size, size_t Alignment);
-
-	// The function used to allocate zeroed memory from the heap.
-	// Must not be null.
-	void* (*AllocateZeroed)(size_t Size, size_t Alignment);
-
-	// The function used to free memory allocated by Allocate() and AllocateZeroed().
-	// Must not be null.
-	void (*Free)(void* Pointer);
-
-	// Function used to log messages using a printf-style format string and va_list arguments.
-	// Strings use UTF-8 encoding.
-	// Must not be null.
-	void (*Log)(const char* File, int Line, void* ProgramCounter, autortfm_log_severity Severity, const char* Format, va_list Args);
-
-	// Function used to log messages with a callstack using a printf-style format string and va_list arguments.
-	// Strings use UTF-8 encoding.
-	// Must not be null.
-	void (*LogWithCallstack)(void* ProgramCounter, autortfm_log_severity Severity, const char* Format, va_list Args);
-
-	// Function used to report an ensure failure using a printf-style format string and va_list arguments.
-	// Strings use UTF-8 encoding.
-	// Must not be null.
-	void (*EnsureFailure)(const char* File, int Line, void* ProgramCounter, const char* Condition, const char* Format, va_list Args);
-
-	// Function used to query whether a log severity is active.
-	// Must not be null.
-	bool (*IsLogActive)(autortfm_log_severity Severity);
-
-	// Optional callback to be informed when the value returned by
-	// ForTheRuntime::IsAutoRTFMRuntimeEnabled() changes.
-	// Can be null.
-	void (*OnRuntimeEnabledChanged)();
-
-	// Optional callback to be informed when the value returned by
-	// ForTheRuntime::GetRetryTransaction() changes.
-	// Can be null.
-	void (*OnRetryTransactionsChanged)();
-
-	// Optional callback to be informed when the value returned by
-	// ForTheRuntime::GetMemoryValidationLevel() changes.
-	// Can be null.
-	void (*OnMemoryValidationLevelChanged)();
-
-	// Optional callback to be informed when the value returned by
-	// ForTheRuntime::GetMemoryValidationThrottlingEnabled() changes.
-	// Can be null.
-	void (*OnMemoryValidationThrottlingChanged)();
-
-	// Optional callback to be informed when the value returned by
-	// ForTheRuntime::GetMemoryValidationStatisticsEnabled() changes.
-	// Can be null.
-	void (*OnMemoryValidationStatisticsChanged)();
-
-} autortfm_extern_api;
-
-#if UE_AUTORTFM_ENABLED
-// Initialize the AutoRTFM library.
-// Parameters:
-//   ExternAPI - Function pointers used by AutoRTFM for heap allocations, etc.
-//               Must be non-null.
-UE_AUTORTFM_API void autortfm_initialize(const autortfm_extern_api* ExternAPI) AUTORTFM_NOEXCEPT;
-#else
-UE_AUTORTFM_CRITICAL_INLINE void autortfm_initialize(const autortfm_extern_api* ExternAPI) AUTORTFM_NOEXCEPT
-{
-	UE_AUTORTFM_UNUSED(ExternAPI);
-}
-#endif
-
-#if UE_AUTORTFM_ENABLED
-// Note: There is no implementation of this function.
-// The AutoRTFM compiler will replace all calls to this function with a constant boolean value.
-UE_AUTORTFM_API bool autortfm_is_closed(void) AUTORTFM_NOEXCEPT;
-#else
-UE_AUTORTFM_CRITICAL_INLINE bool autortfm_is_closed(void) AUTORTFM_NOEXCEPT
-{
-    return false;
-}
-#endif
-
-#if UE_AUTORTFM_ENABLED
-UE_AUTORTFM_API bool autortfm_is_transactional(void) AUTORTFM_NOEXCEPT;
-#else
-UE_AUTORTFM_CRITICAL_INLINE bool autortfm_is_transactional(void) AUTORTFM_NOEXCEPT
-{
-    return false;
-}
-#endif
-
-#if UE_AUTORTFM_ENABLED
-UE_AUTORTFM_API bool autortfm_is_committing(void) AUTORTFM_NOEXCEPT;
-#else
-UE_AUTORTFM_CRITICAL_INLINE bool autortfm_is_committing(void) AUTORTFM_NOEXCEPT
-{
-	return false;
-}
-#endif
-
-#if UE_AUTORTFM_ENABLED
-UE_AUTORTFM_API bool autortfm_is_committing_or_aborting(void) AUTORTFM_NOEXCEPT;
-#else
-UE_AUTORTFM_CRITICAL_INLINE bool autortfm_is_committing_or_aborting(void) AUTORTFM_NOEXCEPT
-{
-	return false;
-}
-#endif
-
-#if UE_AUTORTFM_ENABLED
-UE_AUTORTFM_API bool autortfm_is_retrying(void) AUTORTFM_NOEXCEPT;
-#else
-UE_AUTORTFM_CRITICAL_INLINE bool autortfm_is_retrying(void) AUTORTFM_NOEXCEPT
-{
-	return false;
-}
-#endif
-
-#if UE_AUTORTFM_ENABLED
-UE_AUTORTFM_API autortfm_result autortfm_transact(void (*UninstrumentedWork)(void*), void (*InstrumentedWork)(void*), void* Arg) AUTORTFM_EXCEPT;
-#else
-UE_AUTORTFM_CRITICAL_INLINE autortfm_result autortfm_transact(void (*UninstrumentedWork)(void*), void (*InstrumentedWork)(void*), void* Arg) AUTORTFM_EXCEPT
-{
-	UE_AUTORTFM_UNUSED(InstrumentedWork);
-	UninstrumentedWork(Arg);
-	return autortfm_committed;
-}
-#endif
-
-#if UE_AUTORTFM_ENABLED
-UE_AUTORTFM_API autortfm_result autortfm_transact_then_open(void (*UninstrumentedWork)(void*), void* Arg) AUTORTFM_EXCEPT;
-#else
-UE_AUTORTFM_CRITICAL_INLINE autortfm_result autortfm_transact_then_open(void (*UninstrumentedWork)(void*), void* Arg) AUTORTFM_EXCEPT
-{
-	UninstrumentedWork(Arg);
-    return autortfm_committed;
-}
-#endif
-
-#if UE_AUTORTFM_ENABLED
-UE_AUTORTFM_API void autortfm_commit(void (*UninstrumentedWork)(void*), void (*InstrumentedWork)(void*), void* Arg);
-#else
-UE_AUTORTFM_CRITICAL_INLINE void autortfm_commit(void (*UninstrumentedWork)(void*), void (*InstrumentedWork)(void*), void* Arg)
-{
-	UE_AUTORTFM_UNUSED(InstrumentedWork);
-	UninstrumentedWork(Arg);
-}
-#endif
-
-#if UE_AUTORTFM_ENABLED
-UE_AUTORTFM_API void autortfm_abort_transaction() AUTORTFM_NOEXCEPT;
-#else
-UE_AUTORTFM_CRITICAL_INLINE void autortfm_abort_transaction() AUTORTFM_NOEXCEPT {}
-#endif
-
-#if UE_AUTORTFM_ENABLED
-AUTORTFM_OPEN UE_AUTORTFM_API autortfm_transaction_id autortfm_current_transaction_id() AUTORTFM_NOEXCEPT;
-#else
-UE_AUTORTFM_CRITICAL_INLINE autortfm_transaction_id autortfm_current_transaction_id() AUTORTFM_NOEXCEPT { return 0; }
-#endif
-
-#if UE_AUTORTFM_ENABLED
-UE_AUTORTFM_API autortfm_result autortfm_commit_transaction() AUTORTFM_NOEXCEPT;
-#else
-UE_AUTORTFM_CRITICAL_INLINE autortfm_result autortfm_commit_transaction() AUTORTFM_NOEXCEPT { return autortfm_aborted_by_language; }
-#endif
-
-#if UE_AUTORTFM_ENABLED
-UE_AUTORTFM_FORCENOINLINE UE_AUTORTFM_API void autortfm_open(void (*work)(void* arg), void* arg, const void* return_address);
-#else
-UE_AUTORTFM_CRITICAL_INLINE void autortfm_open(void (*work)(void* arg), void* arg, const void* /* return_address */) { work(arg); }
-#endif
-
-#if UE_AUTORTFM_ENABLED
-UE_AUTORTFM_FORCENOINLINE UE_AUTORTFM_API void autortfm_open_explicit_validation(autortfm_memory_validation_level, void (*work)(void* arg), void* arg, const void* return_address);
-#else
-UE_AUTORTFM_CRITICAL_INLINE void autortfm_open_explicit_validation(autortfm_memory_validation_level, void (*work)(void* arg), void* arg, const void* /* return_address */) { work(arg); }
-#endif
-
-#if UE_AUTORTFM_ENABLED
-[[nodiscard]] UE_AUTORTFM_API autortfm_status autortfm_close(void (*UninstrumentedWork)(void*), void (*InstrumentedWork)(void*), void* Arg);
-#else
-AUTORTFM_DISABLE_UNREACHABLE_CODE_WARNINGS
-[[nodiscard]] UE_AUTORTFM_CRITICAL_INLINE autortfm_status autortfm_close(void (*UninstrumentedWork)(void*), void (*InstrumentedWork)(void*), void* Arg)
-{
-	UE_AUTORTFM_UNUSED(UninstrumentedWork);
-	UE_AUTORTFM_UNUSED(InstrumentedWork);
-	UE_AUTORTFM_UNUSED(Arg);
-	abort();
-	return autortfm_status_aborted_by_language;
-}
-AUTORTFM_RESTORE_UNREACHABLE_CODE_WARNINGS
-#endif
-
-#if UE_AUTORTFM_ENABLED
-UE_AUTORTFM_API void autortfm_record_open_write(void* Ptr, size_t Size) AUTORTFM_NOEXCEPT;
-#else
-UE_AUTORTFM_CRITICAL_INLINE void autortfm_record_open_write(void* Ptr, size_t Size) AUTORTFM_NOEXCEPT
-{
-	UE_AUTORTFM_UNUSED(Ptr);
-	UE_AUTORTFM_UNUSED(Size);
-}
-#endif
-
-#if UE_AUTORTFM_ENABLED
-UE_AUTORTFM_API void autortfm_record_open_write_no_memory_validation(void* Ptr, size_t Size) AUTORTFM_NOEXCEPT;
-#else
-UE_AUTORTFM_CRITICAL_INLINE void autortfm_record_open_write_no_memory_validation(void* Ptr, size_t Size) AUTORTFM_NOEXCEPT
-{
-	UE_AUTORTFM_UNUSED(Ptr);
-	UE_AUTORTFM_UNUSED(Size);
-}
-#endif
-
-// autortfm_open_to_closed_mapping maps an open function to its closed variant.
-struct autortfm_open_to_closed_mapping
-{
-	void* Open;
-	void* Closed;
-};
-
-// autortfm_open_to_closed_table holds a pointer to a null-terminated list of
-// autortfm_open_to_closed_mapping, and an intrusive linked-list pointer to the
-// previous and next registered autortfm_open_to_closed_table.
-struct autortfm_open_to_closed_table
-{
-	// Null-terminated open function to closed function mapping table.
-	const struct autortfm_open_to_closed_mapping* Mappings;
-	// An intrusive linked-list pointer to the previous autortfm_open_to_closed_table.
-	// Used by autortfm_register_open_to_closed_functions().
-	struct autortfm_open_to_closed_table* Prev;
-	// An intrusive linked-list pointer to the next autortfm_open_to_closed_table.
-	// Used by autortfm_register_open_to_closed_functions().
-	struct autortfm_open_to_closed_table* Next;
-};
-
-#if UE_AUTORTFM_ENABLED
-UE_AUTORTFM_API void autortfm_register_open_to_closed_functions(struct autortfm_open_to_closed_table* Table) AUTORTFM_NOEXCEPT;
-UE_AUTORTFM_API void autortfm_unregister_open_to_closed_functions(struct autortfm_open_to_closed_table* Table) AUTORTFM_NOEXCEPT;
-#else
-UE_AUTORTFM_CRITICAL_INLINE void autortfm_register_open_to_closed_functions(autortfm_open_to_closed_table* Table) AUTORTFM_NOEXCEPT
-{
-	UE_AUTORTFM_UNUSED(Table);
-}
-UE_AUTORTFM_CRITICAL_INLINE void autortfm_unregister_open_to_closed_functions(autortfm_open_to_closed_table* Table) AUTORTFM_NOEXCEPT
-{
-	UE_AUTORTFM_UNUSED(Table);
-}
-#endif
-
-#if UE_AUTORTFM_ENABLED
-UE_AUTORTFM_API bool autortfm_is_on_current_transaction_stack(void* Ptr) AUTORTFM_NOEXCEPT;
-#else
-UE_AUTORTFM_CRITICAL_INLINE bool autortfm_is_on_current_transaction_stack(void* Ptr) AUTORTFM_NOEXCEPT
-{
-	UE_AUTORTFM_UNUSED(Ptr);
-	return false;
-}
-#endif
-
-#if UE_AUTORTFM_ENABLED
-UE_AUTORTFM_API void autortfm_on_commit(void (*work)(void* arg), void* arg);
-#else
-UE_AUTORTFM_CRITICAL_INLINE void autortfm_on_commit(void (*work)(void* arg), void* arg)
-{
-    work(arg);
-}
-#endif
-
-#if UE_AUTORTFM_ENABLED
-UE_AUTORTFM_API void autortfm_on_pre_abort(void (*work)(void* arg), void* arg) AUTORTFM_NOEXCEPT;
-#else
-UE_AUTORTFM_CRITICAL_INLINE void autortfm_on_pre_abort(void (*work)(void* arg), void* arg) AUTORTFM_NOEXCEPT
-{
-	UE_AUTORTFM_UNUSED(work);
-	UE_AUTORTFM_UNUSED(arg);
-}
-#endif
-
-#if UE_AUTORTFM_ENABLED
-UE_AUTORTFM_API void autortfm_on_abort(void (*work)(void* arg), void* arg) AUTORTFM_NOEXCEPT;
-#else
-UE_AUTORTFM_CRITICAL_INLINE void autortfm_on_abort(void (*work)(void* arg), void* arg) AUTORTFM_NOEXCEPT
-{
-	UE_AUTORTFM_UNUSED(work);
-	UE_AUTORTFM_UNUSED(arg);
-}
-#endif
-
-#if UE_AUTORTFM_ENABLED
-UE_AUTORTFM_API void autortfm_push_on_abort_handler(const void* key, void (*work)(void* arg), void* arg) AUTORTFM_NOEXCEPT;
-UE_AUTORTFM_API void autortfm_pop_on_abort_handler(const void* key) AUTORTFM_NOEXCEPT;
-#else
-UE_AUTORTFM_CRITICAL_INLINE void autortfm_push_on_abort_handler(const void* key, void (*work)(void* arg), void* arg) AUTORTFM_NOEXCEPT
-{
-	UE_AUTORTFM_UNUSED(key);
-	UE_AUTORTFM_UNUSED(work);
-	UE_AUTORTFM_UNUSED(arg);
-}
-
-UE_AUTORTFM_CRITICAL_INLINE void autortfm_pop_on_abort_handler(const void* key) AUTORTFM_NOEXCEPT
-{
-	UE_AUTORTFM_UNUSED(key);
-}
-#endif
-
-#if UE_AUTORTFM_ENABLED
-UE_AUTORTFM_API void* autortfm_did_allocate(void* ptr, size_t size) AUTORTFM_NOEXCEPT;
-#else
-UE_AUTORTFM_CRITICAL_INLINE void* autortfm_did_allocate(void* ptr, size_t size) AUTORTFM_NOEXCEPT
-{
-	UE_AUTORTFM_UNUSED(size);
-    return ptr;
-}
-#endif
-
-#if UE_AUTORTFM_ENABLED
-UE_AUTORTFM_API void autortfm_did_free(void* ptr) AUTORTFM_NOEXCEPT;
-#else
-UE_AUTORTFM_CRITICAL_INLINE void autortfm_did_free(void* ptr) AUTORTFM_NOEXCEPT
-{
-	UE_AUTORTFM_UNUSED(ptr);
-}
-#endif
-
-// If running with AutoRTFM enabled, then perform an ABI check between the
-// AutoRTFM compiler and the AutoRTFM runtime, to ensure that memory is being
-// laid out in an identical manner between the AutoRTFM runtime and the AutoRTFM
-// compiler pass. Should not be called manually by the user, a call to this will
-// be injected by the compiler into a global constructor in the AutoRTFM compiled
-// code.
-#if UE_AUTORTFM_ENABLED
-UE_AUTORTFM_API void autortfm_check_abi(void* ptr, size_t size) AUTORTFM_NOEXCEPT;
-#else
-UE_AUTORTFM_CRITICAL_INLINE void autortfm_check_abi(void* ptr, size_t size) AUTORTFM_NOEXCEPT
-{
-	UE_AUTORTFM_UNUSED(ptr);
-	UE_AUTORTFM_UNUSED(size);
-}
-#endif
-
-#if UE_AUTORTFM_ENABLED
-// Called when execution unexpectedly reaches a code path that was considered unreachable.
-// Either aborts execution of the program or aborts the current transaction, depending on the
-// current 'InternalAbortAction' state.
-[[noreturn]] UE_AUTORTFM_API void autortfm_unreachable(const char* Message) AUTORTFM_NOEXCEPT;
-#else
-UE_AUTORTFM_CRITICAL_INLINE void autortfm_unreachable(const char* Message) AUTORTFM_NOEXCEPT
-{
-	UE_AUTORTFM_UNUSED(Message);
-}
-#endif
-
-#ifdef __cplusplus
-}
-
 namespace AutoRTFM
 {
 
@@ -598,8 +176,6 @@ template<typename FunctorType>
 void AutoRTFMFunctorInvoker(void* Arg) { (*static_cast<const FunctorType*>(Arg))(); }
 
 #if UE_AUTORTFM
-extern "C" UE_AUTORTFM_FORCENOINLINE UE_AUTORTFM_API void* autortfm_lookup_function(void* OriginalFunction, const char* Where) AUTORTFM_NOEXCEPT;
-
 template<typename FunctorType>
 auto AutoRTFMLookupInstrumentedFunctorInvoker(const FunctorType& Functor) -> void(*)(void*)
 {
diff --git a/Engine/Source/Runtime/AutoRTFM/Public/AutoRTFM/CAPI.h b/Engine/Source/Runtime/AutoRTFM/Public/AutoRTFM/CAPI.h
new file mode 100644
--- /dev/null
+++ b/Engine/Source/Runtime/AutoRTFM/Public/AutoRTFM/CAPI.h
@@ -0,0 +1,443 @@
+// Copyright Epic Games, Inc. All Rights Reserved.
+
+#pragma once
+
+#include "AutoRTFMDefines.h"
+#include <cstdarg> // for va_list
+
+#ifdef __cplusplus
+extern "C"
+{
+#endif
+
+// The C API exists for a few reasons:
+//
+// - It makes linking easy. AutoRTFM has to deal with a weird kind of linking
+//   where the compiler directly emits calls to functions with a given name.
+//   It's easiest to do that in llvm if the functions have C linkage and C ABI.
+// - It makes testing easy. Even seemingly simple C++ code introduces pitfalls
+//   for AutoRTFM. So very focused tests work best when written in C.
+// - It makes compiler optimizations much easier as there is no mangling to
+// 	 consider when looking for functions in the runtime we can optimize.
+//
+// We use snake_case for C API surface area to make it easy to distinguish.
+//
+// The C API should not be used directly - it is here purely as an
+// implementation detail.
+
+// This must match AutoRTFM::ETransactionResult.
+typedef enum
+{
+    autortfm_aborted_by_request = 0,
+    autortfm_aborted_by_language,
+    autortfm_committed,
+	autortfm_aborted_by_transact_during_commit,
+	autortfm_aborted_by_transact_during_abort,
+	autortfm_aborted_by_cascade,
+} autortfm_result;
+
+// This must match AutoRTFM::EContextStatus.
+typedef enum
+{
+	autortfm_status_idle = 0,
+	autortfm_status_ontrack,
+	autortfm_status_aborted_by_failed_lock_aquisition,
+	autortfm_status_aborted_by_language,
+	autortfm_status_aborted_by_request,
+	autortfm_status_committing,
+	autortfm_status_aborted_by_cascading_abort,
+	autortfm_status_aborted_by_cascading_retry,
+	autortfm_status_in_static_local_initializer,
+	autortfm_status_in_post_abort
+} autortfm_status;
+
+// AutoRTFM logging severity.
+typedef enum
+{
+	autortfm_log_verbose = 0,
+	autortfm_log_info,
+	autortfm_log_warn,
+	autortfm_log_error,
+	autortfm_log_fatal,
+} autortfm_log_severity;
+
+// An opaque unique identifier for a transaction.
+typedef uint64_t autortfm_transaction_id;
+
+// Function pointers used by AutoRTFM for heap allocations, etc.
+typedef struct
+{
+	// The function used to allocate memory from the heap.
+	// Must not be null.
+	void* (*Allocate)(size_t Size, size_t Alignment);
+
+	// The function used to reallocate memory from the heap.
+	// Must not be null.
+	void* (*Reallocate)(void* Pointer, size_t Size, size_t Alignment);
+
+	// The function used to allocate zeroed memory from the heap.
+	// Must not be null.
+	void* (*AllocateZeroed)(size_t Size, size_t Alignment);
+
+	// The function used to free memory allocated by Allocate() and AllocateZeroed().
+	// Must not be null.
+	void (*Free)(void* Pointer);
+
+	// Function used to log messages using a printf-style format string and va_list arguments.
+	// Strings use UTF-8 encoding.
+	// Must not be null.
+	void (*Log)(const char* File, int Line, void* ProgramCounter, autortfm_log_severity Severity, const char* Format, va_list Args);
+
+	// Function used to log messages with a callstack using a printf-style format string and va_list arguments.
+	// Strings use UTF-8 encoding.
+	// Must not be null.
+	void (*LogWithCallstack)(void* ProgramCounter, autortfm_log_severity Severity, const char* Format, va_list Args);
+
+	// Function used to report an ensure failure using a printf-style format string and va_list arguments.
+	// Strings use UTF-8 encoding.
+	// Must not be null.
+	void (*EnsureFailure)(const char* File, int Line, void* ProgramCounter, const char* Condition, const char* Format, va_list Args);
+
+	// Function used to query whether a log severity is active.
+	// Must not be null.
+	bool (*IsLogActive)(autortfm_log_severity Severity);
+
+	// Optional callback to be informed when the value returned by
+	// ForTheRuntime::IsAutoRTFMRuntimeEnabled() changes.
+	// Can be null.
+	void (*OnRuntimeEnabledChanged)();
+
+	// Optional callback to be informed when the value returned by
+	// ForTheRuntime::GetRetryTransaction() changes.
+	// Can be null.
+	void (*OnRetryTransactionsChanged)();
+
+	// Optional callback to be informed when the value returned by
+	// ForTheRuntime::GetMemoryValidationLevel() changes.
+	// Can be null.
+	void (*OnMemoryValidationLevelChanged)();
+
+	// Optional callback to be informed when the value returned by
+	// ForTheRuntime::GetMemoryValidationThrottlingEnabled() changes.
+	// Can be null.
+	void (*OnMemoryValidationThrottlingChanged)();
+
+	// Optional callback to be informed when the value returned by
+	// ForTheRuntime::GetMemoryValidationStatisticsEnabled() changes.
+	// Can be null.
+	void (*OnMemoryValidationStatisticsChanged)();
+
+} autortfm_extern_api;
+
+#if UE_AUTORTFM_ENABLED
+// Initialize the AutoRTFM library.
+// Parameters:
+//   ExternAPI - Function pointers used by AutoRTFM for heap allocations, etc.
+//               Must be non-null.
+UE_AUTORTFM_API void autortfm_initialize(const autortfm_extern_api* ExternAPI) AUTORTFM_NOEXCEPT;
+#else
+UE_AUTORTFM_CRITICAL_INLINE void autortfm_initialize(const autortfm_extern_api* ExternAPI) AUTORTFM_NOEXCEPT
+{
+	UE_AUTORTFM_UNUSED(ExternAPI);
+}
+#endif
+
+#if UE_AUTORTFM_ENABLED
+// Note: There is no implementation of this function.
+// The AutoRTFM compiler will replace all calls to this function with a constant boolean value.
+UE_AUTORTFM_API bool autortfm_is_closed(void) AUTORTFM_NOEXCEPT;
+#else
+UE_AUTORTFM_CRITICAL_INLINE bool autortfm_is_closed(void) AUTORTFM_NOEXCEPT
+{
+    return false;
+}
+#endif
+
+#if UE_AUTORTFM_ENABLED
+UE_AUTORTFM_API bool autortfm_is_transactional(void) AUTORTFM_NOEXCEPT;
+#else
+UE_AUTORTFM_CRITICAL_INLINE bool autortfm_is_transactional(void) AUTORTFM_NOEXCEPT
+{
+    return false;
+}
+#endif
+
+#if UE_AUTORTFM_ENABLED
+UE_AUTORTFM_API bool autortfm_is_committing(void) AUTORTFM_NOEXCEPT;
+#else
+UE_AUTORTFM_CRITICAL_INLINE bool autortfm_is_committing(void) AUTORTFM_NOEXCEPT
+{
+	return false;
+}
+#endif
+
+#if UE_AUTORTFM_ENABLED
+UE_AUTORTFM_API bool autortfm_is_committing_or_aborting(void) AUTORTFM_NOEXCEPT;
+#else
+UE_AUTORTFM_CRITICAL_INLINE bool autortfm_is_committing_or_aborting(void) AUTORTFM_NOEXCEPT
+{
+	return false;
+}
+#endif
+
+#if UE_AUTORTFM_ENABLED
+UE_AUTORTFM_API bool autortfm_is_retrying(void) AUTORTFM_NOEXCEPT;
+#else
+UE_AUTORTFM_CRITICAL_INLINE bool autortfm_is_retrying(void) AUTORTFM_NOEXCEPT
+{
+	return false;
+}
+#endif
+
+#if UE_AUTORTFM_ENABLED
+UE_AUTORTFM_API autortfm_result autortfm_transact(void (*UninstrumentedWork)(void*), void (*InstrumentedWork)(void*), void* Arg) AUTORTFM_EXCEPT;
+#else
+UE_AUTORTFM_CRITICAL_INLINE autortfm_result autortfm_transact(void (*UninstrumentedWork)(void*), void (*InstrumentedWork)(void*), void* Arg) AUTORTFM_EXCEPT
+{
+	UE_AUTORTFM_UNUSED(InstrumentedWork);
+	UninstrumentedWork(Arg);
+	return autortfm_committed;
+}
+#endif
+
+#if UE_AUTORTFM_ENABLED
+UE_AUTORTFM_API autortfm_result autortfm_transact_then_open(void (*UninstrumentedWork)(void*), void* Arg) AUTORTFM_EXCEPT;
+#else
+UE_AUTORTFM_CRITICAL_INLINE autortfm_result autortfm_transact_then_open(void (*UninstrumentedWork)(void*), void* Arg) AUTORTFM_EXCEPT
+{
+	UninstrumentedWork(Arg);
+    return autortfm_committed;
+}
+#endif
+
+#if UE_AUTORTFM_ENABLED
+UE_AUTORTFM_API void autortfm_commit(void (*UninstrumentedWork)(void*), void (*InstrumentedWork)(void*), void* Arg);
+#else
+UE_AUTORTFM_CRITICAL_INLINE void autortfm_commit(void (*UninstrumentedWork)(void*), void (*InstrumentedWork)(void*), void* Arg)
+{
+	UE_AUTORTFM_UNUSED(InstrumentedWork);
+	UninstrumentedWork(Arg);
+}
+#endif
+
+#if UE_AUTORTFM_ENABLED
+UE_AUTORTFM_API void autortfm_abort_transaction() AUTORTFM_NOEXCEPT;
+#else
+UE_AUTORTFM_CRITICAL_INLINE void autortfm_abort_transaction() AUTORTFM_NOEXCEPT {}
+#endif
+
+#if UE_AUTORTFM_ENABLED
+AUTORTFM_OPEN UE_AUTORTFM_API autortfm_transaction_id autortfm_current_transaction_id() AUTORTFM_NOEXCEPT;
+#else
+UE_AUTORTFM_CRITICAL_INLINE autortfm_transaction_id autortfm_current_transaction_id() AUTORTFM_NOEXCEPT { return 0; }
+#endif
+
+#if UE_AUTORTFM_ENABLED
+UE_AUTORTFM_API autortfm_result autortfm_commit_transaction() AUTORTFM_NOEXCEPT;
+#else
+UE_AUTORTFM_CRITICAL_INLINE autortfm_result autortfm_commit_transaction() AUTORTFM_NOEXCEPT { return autortfm_aborted_by_language; }
+#endif
+
+#if UE_AUTORTFM_ENABLED
+UE_AUTORTFM_FORCENOINLINE UE_AUTORTFM_API void autortfm_open(void (*work)(void* arg), void* arg, const void* return_address);
+#else
+UE_AUTORTFM_CRITICAL_INLINE void autortfm_open(void (*work)(void* arg), void* arg, const void* /* return_address */) { work(arg); }
+#endif
+
+#if UE_AUTORTFM_ENABLED
+UE_AUTORTFM_FORCENOINLINE UE_AUTORTFM_API void autortfm_open_explicit_validation(autortfm_memory_validation_level, void (*work)(void* arg), void* arg, const void* return_address);
+#else
+UE_AUTORTFM_CRITICAL_INLINE void autortfm_open_explicit_validation(autortfm_memory_validation_level, void (*work)(void* arg), void* arg, const void* /* return_address */) { work(arg); }
+#endif
+
+#if UE_AUTORTFM_ENABLED
+[[nodiscard]] UE_AUTORTFM_API autortfm_status autortfm_close(void (*UninstrumentedWork)(void*), void (*InstrumentedWork)(void*), void* Arg);
+#else
+AUTORTFM_DISABLE_UNREACHABLE_CODE_WARNINGS
+[[nodiscard]] UE_AUTORTFM_CRITICAL_INLINE autortfm_status autortfm_close(void (*UninstrumentedWork)(void*), void (*InstrumentedWork)(void*), void* Arg)
+{
+	UE_AUTORTFM_UNUSED(UninstrumentedWork);
+	UE_AUTORTFM_UNUSED(InstrumentedWork);
+	UE_AUTORTFM_UNUSED(Arg);
+#if defined(__clang__)
+	__builtin_trap();
+#endif
+	return autortfm_status_aborted_by_language;
+}
+AUTORTFM_RESTORE_UNREACHABLE_CODE_WARNINGS
+#endif
+
+#if UE_AUTORTFM_ENABLED
+UE_AUTORTFM_API void autortfm_record_open_write(void* Ptr, size_t Size) AUTORTFM_NOEXCEPT;
+#else
+UE_AUTORTFM_CRITICAL_INLINE void autortfm_record_open_write(void* Ptr, size_t Size) AUTORTFM_NOEXCEPT
+{
+	UE_AUTORTFM_UNUSED(Ptr);
+	UE_AUTORTFM_UNUSED(Size);
+}
+#endif
+
+#if UE_AUTORTFM_ENABLED
+UE_AUTORTFM_API void autortfm_record_open_write_no_memory_validation(void* Ptr, size_t Size) AUTORTFM_NOEXCEPT;
+#else
+UE_AUTORTFM_CRITICAL_INLINE void autortfm_record_open_write_no_memory_validation(void* Ptr, size_t Size) AUTORTFM_NOEXCEPT
+{
+	UE_AUTORTFM_UNUSED(Ptr);
+	UE_AUTORTFM_UNUSED(Size);
+}
+#endif
+
+// autortfm_open_to_closed_mapping maps an open function to its closed variant.
+struct autortfm_open_to_closed_mapping
+{
+	void* Open;
+	void* Closed;
+};
+
+// autortfm_open_to_closed_table holds a pointer to a null-terminated list of
+// autortfm_open_to_closed_mapping, and an intrusive linked-list pointer to the
+// previous and next registered autortfm_open_to_closed_table.
+struct autortfm_open_to_closed_table
+{
+	// Null-terminated open function to closed function mapping table.
+	const struct autortfm_open_to_closed_mapping* Mappings;
+	// An intrusive linked-list pointer to the previous autortfm_open_to_closed_table.
+	// Used by autortfm_register_open_to_closed_functions().
+	struct autortfm_open_to_closed_table* Prev;
+	// An intrusive linked-list pointer to the next autortfm_open_to_closed_table.
+	// Used by autortfm_register_open_to_closed_functions().
+	struct autortfm_open_to_closed_table* Next;
+};
+
+#if UE_AUTORTFM_ENABLED
+UE_AUTORTFM_API void autortfm_register_open_to_closed_functions(struct autortfm_open_to_closed_table* Table) AUTORTFM_NOEXCEPT;
+UE_AUTORTFM_API void autortfm_unregister_open_to_closed_functions(struct autortfm_open_to_closed_table* Table) AUTORTFM_NOEXCEPT;
+#else
+UE_AUTORTFM_CRITICAL_INLINE void autortfm_register_open_to_closed_functions(autortfm_open_to_closed_table* Table) AUTORTFM_NOEXCEPT
+{
+	UE_AUTORTFM_UNUSED(Table);
+}
+UE_AUTORTFM_CRITICAL_INLINE void autortfm_unregister_open_to_closed_functions(autortfm_open_to_closed_table* Table) AUTORTFM_NOEXCEPT
+{
+	UE_AUTORTFM_UNUSED(Table);
+}
+#endif
+
+#if UE_AUTORTFM_ENABLED
+UE_AUTORTFM_API bool autortfm_is_on_current_transaction_stack(void* Ptr) AUTORTFM_NOEXCEPT;
+#else
+UE_AUTORTFM_CRITICAL_INLINE bool autortfm_is_on_current_transaction_stack(void* Ptr) AUTORTFM_NOEXCEPT
+{
+	UE_AUTORTFM_UNUSED(Ptr);
+	return false;
+}
+#endif
+
+#if UE_AUTORTFM_ENABLED
+UE_AUTORTFM_API void autortfm_on_commit(void (*work)(void* arg), void* arg);
+#else
+UE_AUTORTFM_CRITICAL_INLINE void autortfm_on_commit(void (*work)(void* arg), void* arg)
+{
+    work(arg);
+}
+#endif
+
+#if UE_AUTORTFM_ENABLED
+UE_AUTORTFM_API void autortfm_on_pre_abort(void (*work)(void* arg), void* arg) AUTORTFM_NOEXCEPT;
+#else
+UE_AUTORTFM_CRITICAL_INLINE void autortfm_on_pre_abort(void (*work)(void* arg), void* arg) AUTORTFM_NOEXCEPT
+{
+	UE_AUTORTFM_UNUSED(work);
+	UE_AUTORTFM_UNUSED(arg);
+}
+#endif
+
+#if UE_AUTORTFM_ENABLED
+UE_AUTORTFM_API void autortfm_on_abort(void (*work)(void* arg), void* arg) AUTORTFM_NOEXCEPT;
+#else
+UE_AUTORTFM_CRITICAL_INLINE void autortfm_on_abort(void (*work)(void* arg), void* arg) AUTORTFM_NOEXCEPT
+{
+	UE_AUTORTFM_UNUSED(work);
+	UE_AUTORTFM_UNUSED(arg);
+}
+#endif
+
+#if UE_AUTORTFM_ENABLED
+UE_AUTORTFM_API void autortfm_push_on_abort_handler(const void* key, void (*work)(void* arg), void* arg) AUTORTFM_NOEXCEPT;
+UE_AUTORTFM_API void autortfm_pop_on_abort_handler(const void* key) AUTORTFM_NOEXCEPT;
+#else
+UE_AUTORTFM_CRITICAL_INLINE void autortfm_push_on_abort_handler(const void* key, void (*work)(void* arg), void* arg) AUTORTFM_NOEXCEPT
+{
+	UE_AUTORTFM_UNUSED(key);
+	UE_AUTORTFM_UNUSED(work);
+	UE_AUTORTFM_UNUSED(arg);
+}
+
+UE_AUTORTFM_CRITICAL_INLINE void autortfm_pop_on_abort_handler(const void* key) AUTORTFM_NOEXCEPT
+{
+	UE_AUTORTFM_UNUSED(key);
+}
+#endif
+
+#if UE_AUTORTFM_ENABLED
+UE_AUTORTFM_API void* autortfm_did_allocate(void* ptr, size_t size) AUTORTFM_NOEXCEPT;
+#else
+UE_AUTORTFM_CRITICAL_INLINE void* autortfm_did_allocate(void* ptr, size_t size) AUTORTFM_NOEXCEPT
+{
+	UE_AUTORTFM_UNUSED(size);
+    return ptr;
+}
+#endif
+
+#if UE_AUTORTFM_ENABLED
+UE_AUTORTFM_API void autortfm_did_free(void* ptr) AUTORTFM_NOEXCEPT;
+#else
+UE_AUTORTFM_CRITICAL_INLINE void autortfm_did_free(void* ptr) AUTORTFM_NOEXCEPT
+{
+	UE_AUTORTFM_UNUSED(ptr);
+}
+#endif
+
+// If running with AutoRTFM enabled, then perform an ABI check between the
+// AutoRTFM compiler and the AutoRTFM runtime, to ensure that memory is being
+// laid out in an identical manner between the AutoRTFM runtime and the AutoRTFM
+// compiler pass. Should not be called manually by the user, a call to this will
+// be injected by the compiler into a global constructor in the AutoRTFM compiled
+// code.
+#if UE_AUTORTFM_ENABLED
+UE_AUTORTFM_API void autortfm_check_abi(void* ptr, size_t size) AUTORTFM_NOEXCEPT;
+#else
+UE_AUTORTFM_CRITICAL_INLINE void autortfm_check_abi(void* ptr, size_t size) AUTORTFM_NOEXCEPT
+{
+	UE_AUTORTFM_UNUSED(ptr);
+	UE_AUTORTFM_UNUSED(size);
+}
+#endif
+
+#if UE_AUTORTFM_ENABLED
+// Called when execution unexpectedly reaches a code path that was considered unreachable.
+// Either aborts execution of the program or aborts the current transaction, depending on the
+// current 'InternalAbortAction' state.
+[[noreturn]] UE_AUTORTFM_API void autortfm_unreachable(const char* Message) AUTORTFM_NOEXCEPT;
+#else
+UE_AUTORTFM_CRITICAL_INLINE void autortfm_unreachable(const char* Message) AUTORTFM_NOEXCEPT
+{
+	UE_AUTORTFM_UNUSED(Message);
+}
+#endif
+
+#if UE_AUTORTFM_ENABLED
+// Given an open function pointer, returns the equivalent closed function pointer.
+extern "C" UE_AUTORTFM_FORCENOINLINE UE_AUTORTFM_API void* autortfm_lookup_function(void* OriginalFunction, const char* Where) AUTORTFM_NOEXCEPT;
+#else
+UE_AUTORTFM_CRITICAL_INLINE void* autortfm_lookup_function(void* OriginalFunction, const char* Where) AUTORTFM_NOEXCEPT
+{
+	UE_AUTORTFM_UNUSED(OriginalFunction);
+	UE_AUTORTFM_UNUSED(Where);
+	return nullptr;
+}
+#endif
+
+#ifdef __cplusplus
+} // extern "C"
+#endif
