[AutoRTFM] Add new mode to identify the specific write that triggers the AutoRTFM memory validator failure

Add a monotonic counter for all transaction state transitions (open -> closed, etc)

When hitting a regular memory validator hash mismatch, record the transition counter for the bad open, clear the counter and issue a transaction retry.
Once the transaction transition counter reaches the counter for the open that caused the write log hash to change, we:
* Perform a write-log hash snapshot
* Write-protect all the pages for the memory addresses in the write log
* Install a vectored exception handler to perform a hash check after each write to those pages
If the hash no longer compares equal to the snapshotted hash, we've identified the specific instruction that cause the hash to change.

This is currently only implemented for windows.

Also: Drop the outer UE namespace for UE::AutoRTFM::FAutomationTestRunner as it can cause symbol resolve failures.

diff --git a/Engine/Source/Programs/AutoRTFMTests/Private/AutoRTFMTestUtils.h b/Engine/Source/Programs/AutoRTFMTests/Private/AutoRTFMTestUtils.h
--- a/Engine/Source/Programs/AutoRTFMTests/Private/AutoRTFMTestUtils.h
+++ b/Engine/Source/Programs/AutoRTFMTests/Private/AutoRTFMTestUtils.h
@@ -119,6 +119,11 @@ private:
 	const AutoRTFM::EMemoryValidationLevel PrevLevel;
 };
 
+#define AUTORTFM_SCOPED_DISABLE_RETRY() \
+	AutoRTFMTestUtils::FScopedRetry DisableRetry \
+		{AutoRTFM::ForTheRuntime::NoRetry}; \
+	static_assert(true) /* require semicolon */
+
 #define AUTORTFM_SCOPED_DISABLE_MEMORY_VALIDATION() \
 	AutoRTFMTestUtils::ScopedMemoryValidationLevel DisableMemoryValidation \
 		{AutoRTFM::EMemoryValidationLevel::Disabled}; \
diff --git a/Engine/Source/Programs/AutoRTFMTests/Private/AutoRTFMTests.cpp b/Engine/Source/Programs/AutoRTFMTests/Private/AutoRTFMTests.cpp
--- a/Engine/Source/Programs/AutoRTFMTests/Private/AutoRTFMTests.cpp
+++ b/Engine/Source/Programs/AutoRTFMTests/Private/AutoRTFMTests.cpp
@@ -123,6 +123,10 @@ int RunTests(int ArgC, const char* ArgV[])
 	AutoRTFM::ForTheRuntime::SetMemoryValidationLevel(AutoRTFM::EMemoryValidationLevel::Error);
 	AutoRTFM::ForTheRuntime::SetMemoryValidationThrottlingEnabled(false);
 
+#if PLATFORM_WINDOWS
+	AutoRTFM::ForTheRuntime::SetMemoryValidationWriteIdentifierEnabled(true);
+#endif
+
 	// Don't print memory validation stats.
 	AutoRTFM::ForTheRuntime::SetMemoryValidationStatisticsEnabled(false);
 
diff --git a/Engine/Source/Programs/AutoRTFMTests/Private/Open.cpp b/Engine/Source/Programs/AutoRTFMTests/Private/Open.cpp
--- a/Engine/Source/Programs/AutoRTFMTests/Private/Open.cpp
+++ b/Engine/Source/Programs/AutoRTFMTests/Private/Open.cpp
@@ -12,8 +12,9 @@
 namespace
 {
 
-UE_AUTORTFM_ALWAYS_OPEN static void AssignIntPointer(int* Pointer, int Value)
+UE_AUTORTFM_ALWAYS_OPEN static void AssignIntPointer(int* Pointer, int Value, int* LineThatPerformedAssignment = nullptr)
 {
+    if (LineThatPerformedAssignment) { *LineThatPerformedAssignment = __LINE__ + 1; }
     *Pointer = Value;
 }
 
@@ -277,6 +278,10 @@ TEST_CASE("Open.Collision")
     AutoRTFMTestUtils::FCaptureWarningContext WarningContext;
     int I = 0;
     int J = 0;
+    int SafeToAssignTo = 0;
+    int LineThatPerformsConflictingWrite = 0;
+
+#define WRITE_THAT_CAUSES_COLLISION(Statement) Statement; LineThatPerformsConflictingWrite = __LINE__
 
     SECTION("NoCollision")
     {
@@ -412,10 +417,15 @@ TEST_CASE("Open.Collision")
         AutoRTFM::Transact([&]
         {
             I = 42;
-            AutoRTFM::Open([&] { I = 24; });
-        });
+            AutoRTFM::Open([&]
+            {
+                SafeToAssignTo = 10;
+
+                WRITE_THAT_CAUSES_COLLISION(I = 24);
 
-        REQUIRE(WarningContext.HasWarning(AutoRTFMTestUtils::kMemoryModifiedInOpenWarning));
+                SafeToAssignTo = 20;
+            });
+        });
     }
 
     SECTION("Transact(Assign, Transact(Open(Assign)))")
@@ -425,11 +435,16 @@ TEST_CASE("Open.Collision")
             I = 42;
             AutoRTFM::Transact([&]
             {
-                AutoRTFM::Open([&] { I = 24; });
+                AutoRTFM::Open([&]
+                {
+                    SafeToAssignTo = 10;
+                    
+                    WRITE_THAT_CAUSES_COLLISION(I = 24);
+
+                    SafeToAssignTo = 20;
+                });
             });
         });
-
-        REQUIRE(WarningContext.HasWarning(AutoRTFMTestUtils::kMemoryModifiedInOpenWarning));
     }
 
     SECTION("Transact(Assign, Transact(Transact(Open(Assign))))")
@@ -443,13 +458,15 @@ TEST_CASE("Open.Collision")
                 {
                     AutoRTFM::Open([&]
                     {
-                        I = 24;
+                        SafeToAssignTo = 10;
+
+                        WRITE_THAT_CAUSES_COLLISION(I = 24);
+
+                        SafeToAssignTo = 20;
                     });
                 });
             });
         });
-
-        REQUIRE(WarningContext.HasWarning(AutoRTFMTestUtils::kMemoryModifiedInOpenWarning));
     }
 
     SECTION("Transact(Transact(Assign, Transact(Open(Assign))))")
@@ -463,13 +480,15 @@ TEST_CASE("Open.Collision")
                 {
                     AutoRTFM::Open([&]
                     {
-                        I = 24;
+                        SafeToAssignTo = 10;
+
+                        WRITE_THAT_CAUSES_COLLISION(I = 24);
+
+                        SafeToAssignTo = 20;
                     });
                 });
             });
         });
-
-        REQUIRE(WarningContext.HasWarning(AutoRTFMTestUtils::kMemoryModifiedInOpenWarning));
     }
 
     SECTION("Transact(Transact(Assign, Open(Assign, Transact())))")
@@ -479,12 +498,15 @@ TEST_CASE("Open.Collision")
             I = 42;
             AutoRTFM::Open([&]
             {
-                I = 24;
+                SafeToAssignTo = 10;
+
+                WRITE_THAT_CAUSES_COLLISION(I = 24);
+
+                SafeToAssignTo = 20;
+
                 AutoRTFM::Transact([&] {});
             });
         });
-
-        REQUIRE(WarningContext.HasWarning(AutoRTFMTestUtils::kMemoryModifiedInOpenWarning));
     }
 
     SECTION("Transact(Transact(Assign, Open(Transact(), Assign)))")
@@ -495,11 +517,14 @@ TEST_CASE("Open.Collision")
             AutoRTFM::Open([&]
             {
                 AutoRTFM::Transact([&] {});
-                I = 24;
+
+                SafeToAssignTo = 10;
+
+                WRITE_THAT_CAUSES_COLLISION(I = 24);
+
+                SafeToAssignTo = 20;
             });
         });
-
-        REQUIRE(WarningContext.HasWarning(AutoRTFMTestUtils::kMemoryModifiedInOpenWarning));
     }
 
     SECTION("Transact(Assign, CallOpen(Assign))")
@@ -507,10 +532,8 @@ TEST_CASE("Open.Collision")
         AutoRTFM::Transact([&]
         {
             I = 42;
-            AssignIntPointer(&I, 24);
+            AssignIntPointer(&I, 24, &LineThatPerformsConflictingWrite);
         });
-
-        REQUIRE(WarningContext.HasWarning(AutoRTFMTestUtils::kMemoryModifiedInOpenWarning));
     }
 
     SECTION("Transact(Assign, Transact(CallOpen(Assign)))")
@@ -520,11 +543,9 @@ TEST_CASE("Open.Collision")
             I = 42;
             AutoRTFM::Transact([&]
             {
-                AssignIntPointer(&I, 24);
+                AssignIntPointer(&I, 24, &LineThatPerformsConflictingWrite);
             });
         });
-
-        REQUIRE(WarningContext.HasWarning(AutoRTFMTestUtils::kMemoryModifiedInOpenWarning));
     }
 
     SECTION("Transact(Assign, OpenNoValidation(Assign))")
@@ -568,12 +589,15 @@ TEST_CASE("Open.Collision")
             I = 42;
             AutoRTFM::Open([&]
             {
-                I = 10;
+                SafeToAssignTo = 10;
+
+                WRITE_THAT_CAUSES_COLLISION(I = 24);
+
+                SafeToAssignTo = 20;
+
                 AutoRTFM::Open<AutoRTFM::EMemoryValidationLevel::Disabled>([&]{});
             });
         });
-
-        REQUIRE(WarningContext.HasWarning(AutoRTFMTestUtils::kMemoryModifiedInOpenWarning));
     }
 
     SECTION("Transact(Assign, Open(Assign, Close(OpenNoValidation())))")
@@ -583,15 +607,18 @@ TEST_CASE("Open.Collision")
             I = 42;
             AutoRTFM::Open([&]
             {
-                I = 10;
+                SafeToAssignTo = 10;
+
+                WRITE_THAT_CAUSES_COLLISION(I = 24);
+
+                SafeToAssignTo = 20;
+
                 std::ignore = AutoRTFM::Close([&]
                 {
                     AutoRTFM::Open<AutoRTFM::EMemoryValidationLevel::Disabled>([&]{});
                 });
             });
         });
-
-        REQUIRE(WarningContext.HasWarning(AutoRTFMTestUtils::kMemoryModifiedInOpenWarning));
     }
 
     SECTION("Transact(Open(RecordOpenWrite()), Open(Assign))")
@@ -604,12 +631,13 @@ TEST_CASE("Open.Collision")
             });
             AutoRTFM::Open([&]
             {
-                I = 42;
+                SafeToAssignTo = 10;
+
+                WRITE_THAT_CAUSES_COLLISION(I = 24);
+
+                SafeToAssignTo = 20;
             });
         });
-
-        REQUIRE(I == 42);
-        REQUIRE(WarningContext.HasWarning(AutoRTFMTestUtils::kMemoryModifiedInOpenWarning));
     }
 
     SECTION("Transact(Open(RecordOpenWriteNoMemoryValidation()), Open(RecordOpenWrite()), Open(Assign))")
@@ -626,11 +654,13 @@ TEST_CASE("Open.Collision")
             });
             AutoRTFM::Open([&]
             {
-                I = 42;
+                SafeToAssignTo = 10;
+
+                WRITE_THAT_CAUSES_COLLISION(I = 24);
+
+                SafeToAssignTo = 20;
             });
         });
-
-        REQUIRE(WarningContext.HasWarning(AutoRTFMTestUtils::kMemoryModifiedInOpenWarning));
     }
 
     SECTION("Transact(Transact(Open(RecordOpenWriteNoMemoryValidation())), Open(RecordOpenWrite()), Open(Assign))")
@@ -650,11 +680,13 @@ TEST_CASE("Open.Collision")
             });
             AutoRTFM::Open([&]
             {
-                I = 42;
+                SafeToAssignTo = 10;
+
+                WRITE_THAT_CAUSES_COLLISION(I = 24);
+
+                SafeToAssignTo = 20;
             });
         });
-
-        REQUIRE(WarningContext.HasWarning(AutoRTFMTestUtils::kMemoryModifiedInOpenWarning));
     }
 
     SECTION("Transact(Open(RecordOpenWriteNoMemoryValidation()), Open(RecordOpenWrite()), Open(Assign)) <large>")
@@ -672,10 +704,78 @@ TEST_CASE("Open.Collision")
             });
             AutoRTFM::Open([&]
             {
-                LargeStruct.V[255] = 42;
+                SafeToAssignTo = 10;
+
+                WRITE_THAT_CAUSES_COLLISION(LargeStruct.V[255] = 42);
+
+                SafeToAssignTo = 20;
             });
         });
+    }
 
-        REQUIRE(WarningContext.HasWarning(AutoRTFMTestUtils::kMemoryModifiedInOpenWarning));
+    if (LineThatPerformsConflictingWrite > 0)
+    {
+        REQUIRE(WarningContext.HasWarningSubstring(AutoRTFMTestUtils::kMemoryModifiedInOpenWarning));
+        if (AutoRTFM::ForTheRuntime::GetMemoryValidationWriteIdentifierEnabled())
+        {
+            std::string ExpectedFileAndLine = std::string(__FILE__ ":") + std::to_string(LineThatPerformsConflictingWrite);
+            REQUIRE(WarningContext.HasWarningSubstring(AutoRTFMTestUtils::kMemoryModifiedInOpenWarning));
+        }
     }
 }
+
+TEST_CASE("Open.Collision.NonDeterministic")
+{
+    if (!AutoRTFM::ForTheRuntime::GetMemoryValidationWriteIdentifierEnabled())
+    {
+        return;
+    }
+
+    AUTORTFM_SCOPED_DISABLE_RETRY(); // Adds unnecessary complexity to this test.
+
+    AUTORTFM_SCOPED_ENABLE_MEMORY_VALIDATION_AS_WARNING();
+    AutoRTFMTestUtils::FCaptureWarningContext WarningContext;
+
+    int OpenCounter = 0;
+    int MemoryCollision = 0;
+    int SafeToAssignTo = 0;
+    bool RanUntilEnd = false;
+
+    AutoRTFM::Transact([&]
+    {
+        REQUIRE(!RanUntilEnd);
+        MemoryCollision = 42;
+
+        AutoRTFM::Open([&]
+        {
+            OpenCounter++;
+        });
+
+        // OpenCounter is incremented in the open, when the memory validator
+        // attempts a cascading retry to identify the specific write that caused
+        // the hash mismatch this will not be 1.
+        if (OpenCounter == 1)
+        {
+            // First time the transaction is called, we'll take this path.
+            AutoRTFM::Open([&]
+            {
+                // Trigger a memory validation failure
+                MemoryCollision = 24;
+            });
+        }
+        else
+        {
+            // When re-running the cascading retry to identify the offending
+            // write, we'll take this path.
+            AutoRTFM::Open([&]
+            {
+                SafeToAssignTo = 24;
+            });
+        }
+
+        RanUntilEnd = true;
+    });
+
+    REQUIRE(RanUntilEnd);
+    REQUIRE(WarningContext.HasWarningSubstring("Failed to find bad write in open function"));
+}
diff --git a/Engine/Source/Programs/AutoRTFMTests/Private/OpenAPI.cpp b/Engine/Source/Programs/AutoRTFMTests/Private/OpenAPI.cpp
--- a/Engine/Source/Programs/AutoRTFMTests/Private/OpenAPI.cpp
+++ b/Engine/Source/Programs/AutoRTFMTests/Private/OpenAPI.cpp
@@ -45,7 +45,7 @@ TEST_CASE("OpenAPI.StartAbortAndStartAgain")
 	REQUIRE(ValueB == 20);
 	REQUIRE(ValueC == 0);
 
-	REQUIRE(WarningContext.HasWarning(AutoRTFMTestUtils::kMemoryModifiedInOpenWarning));
+	REQUIRE(WarningContext.HasWarningSubstring(AutoRTFMTestUtils::kMemoryModifiedInOpenWarning));
 }
 
 TEST_CASE("OpenAPI.CommitScopedFromOpen_Illegal", "[.]")
@@ -806,7 +806,7 @@ TEST_CASE("OpenAPI.Footgun1")
 	REQUIRE(ValueA == 10);
 	REQUIRE(ValueB == 0);
 
-	REQUIRE(WarningContext.HasWarning(AutoRTFMTestUtils::kMemoryModifiedInOpenWarning));
+	REQUIRE(WarningContext.HasWarningSubstring(AutoRTFMTestUtils::kMemoryModifiedInOpenWarning));
 }
 
 TEST_CASE("OpenAPI.Footgun2")
@@ -843,7 +843,7 @@ TEST_CASE("OpenAPI.Footgun2")
 	REQUIRE(ValueB == 0);
 	REQUIRE(ValueC == 10);
 
-	REQUIRE(WarningContext.HasWarning(AutoRTFMTestUtils::kMemoryModifiedInOpenWarning));
+	REQUIRE(WarningContext.HasWarningSubstring(AutoRTFMTestUtils::kMemoryModifiedInOpenWarning));
 }
 
 #if 0
diff --git a/Engine/Source/Runtime/AutoRTFM/Private/API.cpp b/Engine/Source/Runtime/AutoRTFM/Private/API.cpp
--- a/Engine/Source/Runtime/AutoRTFM/Private/API.cpp
+++ b/Engine/Source/Runtime/AutoRTFM/Private/API.cpp
@@ -38,6 +38,7 @@ namespace
 	AutoRTFM::EMemoryValidationLevel GMemoryValidationLevel = GDefaultMemoryValidationLevel;
 
 	bool GMemoryValidationThrottlingEnabled = true;
+	bool GMemoryValidationWriteIdentifierEnabled = false;
 	bool GMemoryValidationStatisticsEnabled = true;
 	bool GAutoRTFMEnsureOnInternalAbort = true;
 
@@ -78,6 +79,17 @@ namespace AutoRTFM
 	} // Testing
 } // AutoRTFM
 
+
+#if UE_AUTORTFM
+namespace AutoRTFM
+{
+	AUTORTFM_INTERNAL bool IsAutoRTFMInitialized()
+	{
+		return FContext::Get() != nullptr;
+	}
+}
+#endif
+
 namespace AutoRTFM
 {
 	namespace ForTheRuntime
@@ -253,12 +265,10 @@ namespace AutoRTFM
 #endif
 		}
 
-		void SetEnsureOnInternalAbort(bool bEnabled)
+		void SetEnsureOnInternalAbort([[maybe_unused]] bool bEnabled)
 		{
 #if UE_AUTORTFM
 			GAutoRTFMEnsureOnInternalAbort = bEnabled;
-#else
-			UE_AUTORTFM_UNUSED(bEnabled);
 #endif
 		}
 
@@ -356,7 +366,7 @@ namespace AutoRTFM
 #endif
 		}
 
-		void SetMemoryValidationThrottlingEnabled(bool bEnabled)
+		void SetMemoryValidationThrottlingEnabled([[maybe_unused]] bool bEnabled)
 		{
 #if UE_AUTORTFM
 			if (GMemoryValidationThrottlingEnabled != bEnabled)
@@ -367,8 +377,37 @@ namespace AutoRTFM
 					GExternAPI.OnMemoryValidationThrottlingChanged();
 				}
 			}
+#endif
+		}
+
+		bool GetMemoryValidationWriteIdentifierEnabled()
+		{
+#if UE_AUTORTFM && PLATFORM_WINDOWS
+			return GMemoryValidationWriteIdentifierEnabled;
 #else
-			UE_AUTORTFM_UNUSED(bEnabled);
+			return false;
+#endif
+		}
+
+		void SetMemoryValidationWriteIdentifierEnabled([[maybe_unused]] bool bEnabled)
+		{
+#if UE_AUTORTFM
+#	if !PLATFORM_WINDOWS
+			if (IsAutoRTFMInitialized())
+			{
+				AUTORTFM_WARN("AutoRTFM memory validation write identifier is not currently supported on this platform");
+			}
+			return;
+#	endif
+
+			if (GMemoryValidationWriteIdentifierEnabled != bEnabled)
+			{
+				GMemoryValidationWriteIdentifierEnabled = bEnabled;
+				if (GExternAPI.OnMemoryValidationWriteIdentifierChanged)
+				{
+					GExternAPI.OnMemoryValidationWriteIdentifierChanged();
+				}
+			}
 #endif
 		}
 
@@ -381,7 +420,7 @@ namespace AutoRTFM
 #endif
 		}
 
-		void SetMemoryValidationStatisticsEnabled(bool bEnabled)
+		void SetMemoryValidationStatisticsEnabled([[maybe_unused]] bool bEnabled)
 		{
 #if UE_AUTORTFM
 			if (GMemoryValidationStatisticsEnabled != bEnabled)
@@ -392,8 +431,6 @@ namespace AutoRTFM
 					GExternAPI.OnMemoryValidationStatisticsChanged();
 				}
 			}
-#else
-			UE_AUTORTFM_UNUSED(bEnabled);
 #endif
 		}
 
@@ -714,11 +751,6 @@ AUTORTFM_INTERNAL void RTFM_autortfm_did_free(void* Ptr)
 	AutoRTFM::InternalUnreachable();
 }
 
-AUTORTFM_INTERNAL bool IsAutoRTFMInitialized()
-{
-	return FContext::Get() != nullptr;
-}
-
 // Consume the GPendingOpenToClosedRegistrations linked list to register the
 // open -> closed functions. This is done via a linked-list to avoid heap
 // allocations before AutoRTFM is initialized.
diff --git a/Engine/Source/Runtime/AutoRTFM/Private/Context.cpp b/Engine/Source/Runtime/AutoRTFM/Private/Context.cpp
--- a/Engine/Source/Runtime/AutoRTFM/Private/Context.cpp
+++ b/Engine/Source/Runtime/AutoRTFM/Private/Context.cpp
@@ -11,6 +11,7 @@
 #include "FunctionMap.h"
 #include "ScopedGuard.h"
 #include "StackRange.h"
+#include "StateTransitionCounter.h"
 #include "Stats.h"
 #include "Transaction.h"
 #include "TransactionInlines.h"
@@ -215,11 +216,17 @@ EContextStatus FContext::CallClosedNest(void (*ClosedFunction)(void* Arg), void*
 	AUTORTFM_ASSERT(FTransaction::EState::OpenActive == Transaction->State());
 	EMemoryValidationLevel PreviousValidationLevel = Transaction->MemoryValidationLevel();
 	const void* PreviousOpenReturnAddress = Transaction->OpenReturnAddress();
-	Transaction->SetClosedActive();
 
 	PushCallNest(CallNestPool.Take(this));
 
-	CurrentNest->Try([&]() { ClosedFunction(Arg); });
+	CurrentNest->Try([&]
+	{ 
+		// Change the transaction state within the try block, as a memory
+		// validation failure may trigger a cascading retry, which will throw.
+		Transaction->SetClosedActive();
+
+		ClosedFunction(Arg);
+	});
 
 	PopCallNest();
 
@@ -410,6 +417,9 @@ ETransactionResult FContext::Transact(void (*UninstrumentedFunction)(void*), voi
 
 	if (!CurrentTransaction)
 	{
+		// Start of an outermost closed-scope transaction.
+		StateTransitionCounter.Reset();
+
 		// If exceptions are enabled, then ensure that the transaction is automatically committed if
 		// an exception is thrown inside the transaction and the handler is outside the transaction.
 		struct FAutoCommitter final
@@ -467,13 +477,15 @@ ETransactionResult FContext::Transact(void (*UninstrumentedFunction)(void*), voi
 
 		bool bTriedToRunOnce = false;
 
-        for (;;)
-        {
-            Status = EContextStatus::OnTrack;
-            AUTORTFM_ASSERT(CurrentTransaction->IsFresh());
+		for (;;)
+		{
+			Status = EContextStatus::OnTrack;
+			AUTORTFM_ASSERT(CurrentTransaction->IsFresh());
+			AUTORTFM_ASSERT(CurrentTransaction->State() == FTransaction::EState::ClosedActive);
+			AUTORTFM_ASSERT(StateTransitionCounter == FStateTransitionCounter{1});
 			CurrentNest->Try([&] () { InstrumentedFunction(Arg); });
 			AUTORTFM_ASSERT(CurrentTransaction == NewTransaction); // The transaction lambda should have unwound any nested transactions.
-            AUTORTFM_ASSERT(Status != EContextStatus::Idle);
+			AUTORTFM_ASSERT(Status != EContextStatus::Idle);
 
 			switch (Status)
 			{
@@ -488,6 +500,7 @@ ETransactionResult FContext::Transact(void (*UninstrumentedFunction)(void*), voi
 					Status = EContextStatus::AbortedByFailedLockAcquisition;
 					CurrentTransaction->AbortWithoutThrowing();
 					ClearTransactionStatus();
+					StateTransitionCounter.Reset();
 
 					// We've tried to run at least once if we get here!
 					CurrentTransaction->Reset();
@@ -523,6 +536,7 @@ ETransactionResult FContext::Transact(void (*UninstrumentedFunction)(void*), voi
 			case EContextStatus::AbortedByCascadingRetry:
 				// Clean up the transaction to get it ready for re-execution.
 				ClearTransactionStatus();
+				StateTransitionCounter.Reset();
 				CurrentTransaction->Reset();
 
 				AUTORTFM_ASSERT(Status == EContextStatus::OnTrack);
@@ -543,6 +557,7 @@ ETransactionResult FContext::Transact(void (*UninstrumentedFunction)(void*), voi
 				continue;
 
 			case EContextStatus::AbortedByFailedLockAcquisition:
+				StateTransitionCounter.Reset();
 				continue; // Retry the transaction
 
 			default:
diff --git a/Engine/Source/Runtime/AutoRTFM/Private/Context.h b/Engine/Source/Runtime/AutoRTFM/Private/Context.h
--- a/Engine/Source/Runtime/AutoRTFM/Private/Context.h
+++ b/Engine/Source/Runtime/AutoRTFM/Private/Context.h
@@ -4,6 +4,7 @@
 
 #if (defined(__AUTORTFM) && __AUTORTFM)
 
+#include "AutoRTFMTask.h"
 #include "CallNest.h"
 #include "IntrusivePool.h"
 #include "OpenHashThrottler.h"
@@ -131,6 +132,28 @@ public:
 	inline FTaskPool& GetTaskPool() { return TaskPool; }
 	inline FTransactionPool& GetTransactionPool() { return TransactionPool; }
 
+	// Called whenever a transaction changes state.
+	// Increments the monotonic state transition count, and calls the state
+	// transition callback, if set.
+	inline void StateTransitionChanged()
+	{
+		StateTransitionCounter.Increment();
+		if (AUTORTFM_UNLIKELY(StateTransitionCallback.IsSet()))
+		{
+			StateTransitionCallback();
+		}
+	}
+
+	// Returns the total number of times a transaction has transitioned state
+	// since the outermost transaction began.
+	inline FStateTransitionCounter GetStateTransitionCounter() const { return StateTransitionCounter; }
+
+	// Sets a callback to be called whenever the StateTransitionChanged() is called.
+	inline void OnStateTransitionChanged(TTask<void()> Callback) { StateTransitionCallback = Callback; }
+
+	// Returns the stack address limits for the caller thread.
+	static FStackRange GetThreadStackRange();
+
 private:
 	UE_AUTORTFM_API static FContext* Instance;
 
@@ -152,8 +175,6 @@ private:
     // All of this other stuff ought to be private?
     void Reset();
 
-	static FStackRange GetThreadStackRange();
-
 	// We defer allocating FTransactions at the top of the transaction stack.
 	// This allows us to make starting a transaction in the open a load, some math, and a store.
 	uint64_t NumDeferredTransactions{0};
@@ -172,6 +193,8 @@ private:
 		/* AdjustThrottleInterval */ 0.5,  // Adjust throttling probabilities every 500ms.
 		/* TargetFractionHashing  */ 0.1   // At most we want to spent 10% of the time hashing.
 	};
+	FStateTransitionCounter StateTransitionCounter;
+	TTask<void()> StateTransitionCallback;
 };
 
 } // namespace AutoRTFM
diff --git a/Engine/Source/Runtime/AutoRTFM/Private/ExternAPI.cpp b/Engine/Source/Runtime/AutoRTFM/Private/ExternAPI.cpp
--- a/Engine/Source/Runtime/AutoRTFM/Private/ExternAPI.cpp
+++ b/Engine/Source/Runtime/AutoRTFM/Private/ExternAPI.cpp
@@ -74,6 +74,7 @@ autortfm_extern_api GExternAPI
 	/* OnRetryTransactionsChanged */ nullptr,
 	/* OnMemoryValidationLevelChanged */ nullptr,
 	/* OnMemoryValidationThrottlingChanged */ nullptr,
+	/* OnMemoryValidationWriteIdentifierChanged */ nullptr,
 	/* OnMemoryValidationStatisticsChanged */ nullptr,
 };
 
diff --git a/Engine/Source/Runtime/AutoRTFM/Private/StateTransitionCounter.h b/Engine/Source/Runtime/AutoRTFM/Private/StateTransitionCounter.h
new file mode 100644
--- /dev/null
+++ b/Engine/Source/Runtime/AutoRTFM/Private/StateTransitionCounter.h
@@ -0,0 +1,30 @@
+// Copyright Epic Games, Inc. All Rights Reserved.
+
+#pragma once
+
+#if (defined(__AUTORTFM) && __AUTORTFM)
+
+#include "AutoRTFM.h"
+
+namespace AutoRTFM
+{
+
+// A monotonic counter that is incremented every time a transaction changes
+// state (open -> closed, closed -> done), etc.
+// Reset when the outermost transaction is complete.
+struct AUTORTFM_DISABLE FStateTransitionCounter final
+{
+	using CounterType = uint64_t;
+	uint64_t Value = 0;
+	inline void Increment() { Value++; }
+	inline void Reset() { Value = 0; }
+	inline bool operator == (FStateTransitionCounter Other) const { return Value == Other.Value; }
+	inline bool operator <  (FStateTransitionCounter Other) const { return Value < Other.Value; }
+	inline bool operator <= (FStateTransitionCounter Other) const { return Value <= Other.Value; }
+	inline bool operator >  (FStateTransitionCounter Other) const { return Value > Other.Value; }
+	inline bool operator >= (FStateTransitionCounter Other) const { return Value >= Other.Value; }
+};
+
+} // namespace AutoRTFM
+
+#endif // (defined(__AUTORTFM) && __AUTORTFM)
diff --git a/Engine/Source/Runtime/AutoRTFM/Private/Transaction.cpp b/Engine/Source/Runtime/AutoRTFM/Private/Transaction.cpp
--- a/Engine/Source/Runtime/AutoRTFM/Private/Transaction.cpp
+++ b/Engine/Source/Runtime/AutoRTFM/Private/Transaction.cpp
@@ -6,10 +6,15 @@
 #include "TransactionInlines.h"
 
 #include "CallNestInlines.h"
+#include "ContextInlines.h"
 #include "ContextStatus.h"
 #include "HitSet.h"
 #include "Utils.h"
 
+#if PLATFORM_WINDOWS
+#include "WindowsHeader.h"
+#endif
+
 #include <cstdarg>
 #include <cstddef>
 #include <cstdint>
@@ -22,8 +27,7 @@ void FTransaction::Resurrect(FContext* const InContext)
 	AUTORTFM_ASSERT(Context == InContext)
 	Parent = nullptr;
 	StatDepth = 1;
-	RecordedWriteHash = 0;
-	NumWriteLogsHashed = 0;
+	WriteLogSnapshot = {};
 	CurrentMemoryValidationLevel = EMemoryValidationLevel::Disabled;
 	CurrentOpenReturnAddress = nullptr;
 	CurrentState = EState::Uninitialized;
@@ -459,8 +463,7 @@ void FTransaction::SetOpenActive(EMemoryValidationLevel NewMemoryValidationLevel
 	CurrentOpenReturnAddress = ReturnAddress;
 	
 	// TODO: Validate if open -> open with different validation levels.
-	RecordedWriteHash = 0;
-	NumWriteLogsHashed = 0;
+	WriteLogSnapshot = {};
 	CurrentOpenReturnAddress = nullptr;
 
 	SetState<EState::OpenActive>();
@@ -526,6 +529,8 @@ template<FTransaction::EState NewState>
 void FTransaction::SetState()
 {
 	AUTORTFM_ASSERT(NewState != CurrentState);
+	
+	Context->StateTransitionChanged();
 
 	switch (CurrentState)
 	{
@@ -539,12 +544,11 @@ void FTransaction::SetState()
 			if (CurrentMemoryValidationLevel != EMemoryValidationLevel::Disabled)
 			{
 				ValidateWriteHash();
-				RecordedWriteHash = 0;
-				NumWriteLogsHashed = 0;
+				WriteLogSnapshot = {};
 			}
 			else
 			{
-				AUTORTFM_ASSERT(RecordedWriteHash == 0 && NumWriteLogsHashed == 0);
+				AUTORTFM_ASSERT(WriteLogSnapshot.IsZero());
 			}
 			break;
 
@@ -576,7 +580,7 @@ void FTransaction::SetState()
 	if (NewState == EState::OpenActive) {
 		if (CurrentMemoryValidationLevel != EMemoryValidationLevel::Disabled)
 		{
-			AUTORTFM_ASSERT(RecordedWriteHash == 0 && NumWriteLogsHashed == 0);
+			AUTORTFM_ASSERT(WriteLogSnapshot.IsZero());
 			RecordWriteHash();
 		}
 	}
@@ -588,7 +592,7 @@ void FTransaction::DebugBreakIfMemoryValidationFails()
 {
 	if (CurrentMemoryValidationLevel != EMemoryValidationLevel::Disabled)
 	{
-		FWriteHash OldHash = RecordedWriteHash;
+		FWriteHash OldHash = WriteLogSnapshot.Hash;
 		FWriteHash NewHash = CalculateNestedWriteHash(CurrentOpenReturnAddress);
 		if (OldHash != NewHash)
 		{
@@ -600,24 +604,264 @@ void FTransaction::DebugBreakIfMemoryValidationFails()
 
 void FTransaction::RecordWriteHash()
 {
-	NumWriteLogsHashed = WriteLog.Num();
-	RecordedWriteHash = CalculateNestedWriteHash(CurrentOpenReturnAddress);
+	WriteLogSnapshot.StateTransitionCounter = Context->GetStateTransitionCounter();
+	WriteLogSnapshot.NumWriteLogsHashed = WriteLog.Num();
+	WriteLogSnapshot.Hash = CalculateNestedWriteHash(CurrentOpenReturnAddress);
 }
 
-void FTransaction::ValidateWriteHash() const
+#if PLATFORM_WINDOWS
+// Used to identify the specific open write that causes a memory validation hash mismatch.
+struct AUTORTFM_DISABLE FTransaction::FBadWriteIdentifier
+{
+	// Constructor
+	// OpenCounter is the transaction state transition counter at the transition
+	// to the open that caused the write log hash to change. 
+	FBadWriteIdentifier(FStateTransitionCounter OpenCounter)
+		: OpenCounter(OpenCounter)
+	{
+		// Disable the automatic open <-> closed memory validation while
+		// attempting to find the specific write.
+		AutoRTFM::ForTheRuntime::SetMemoryValidationLevel(EMemoryValidationLevel::Disabled);
+	}
+
+	// Destructor
+	// Performs the following clean up:
+	// * Uninstalls the vectored exception handler and restores the page write
+	//   protection bits (if required).
+	// * Removes the state transition callback from the FContext.
+	// * Restores the original memory validation level.
+	~FBadWriteIdentifier()
+	{
+		if (ExceptionHandler)
+		{
+			UninstallExceptionHandler();
+			WriteUnprotectPages();
+		}
+		Context.OnStateTransitionChanged(TTask<void()>{});
+		AutoRTFM::ForTheRuntime::SetMemoryValidationLevel(OldMemoryValidationLevel);
+	}
+
+	// The handler for all transaction state changes.
+	// This is called by the callback passed to Context.OnStateTransitionChanged().
+	// When the state transition counter value is equal to the counter value
+	// observed for the open that caused the write log hash to change, we
+	// perform a write-log hash snapshot, then write-protect all the pages for
+	// the memory addresses in the write log, and install a vectored exception
+	// handler to perform a hash check after each write to those pages.
+	// When the state transition counter exceeds the counter value observed for
+	// the open that caused the write log hash to change, we delete this
+	// instance which will clean up all the things we changed.
+	void OnTransactionStateChange()
+	{
+		if (Context.IsAborting())
+		{
+			// Likely called while cascade-aborting after the memory validator
+			// fired. Wait until we've restarted the outer transaction.
+			return;
+		}
+
+		const FStateTransitionCounter Counter = Context.GetStateTransitionCounter();
+		if (Counter == OpenCounter)
+		{
+			// Counter has reached the value observed at the start of the bad
+			// open. Write-protect the memory and install the VEH.
+			AUTORTFM_WARN("Searching for bad write in open function...");
+			FTransaction* Transaction = Context.GetCurrentTransaction();
+			NumWriteLogsHashed = Transaction->WriteLog.Num();
+			Hash = Transaction->CalculateNestedWriteHashWithLimit(NumWriteLogsHashed, nullptr);
+			WriteProtectPages();
+			InstallExceptionHandler();
+		}
+		else if (Counter > OpenCounter)
+		{
+			// Counter has exceeded the value observed at the start of the bad
+			// open. Time to self-destruct.
+			AUTORTFM_WARN("Failed to find bad write in open function");
+			delete this;
+		}
+	}
+
+	// Called by the VEH at the first occurrence of a hash mismatch.
+	// Dumps the callstack, and then self-destructs.
+	void OnBadWrite()
+	{
+		OnValidatorWriteHashMismatch(OldMemoryValidationLevel, "Open write to transactional memory found");
+		delete this;
+	}
+
+	// Changes all the virtual memory pages in the addresses of the write log to
+	// PAGE_READONLY.
+	void WriteProtectPages()
+	{
+		ChangeWriteLogPageProtections(PAGE_READONLY);
+	}
+
+	// Changes all the virtual memory pages in the addresses of the write log to
+	// PAGE_READWRITE.
+	void WriteUnprotectPages()
+	{
+		ChangeWriteLogPageProtections(PAGE_READWRITE);
+	}
+
+	// Changes all the virtual memory pages in the addresses of the write log to
+	// NewProtect.
+	void ChangeWriteLogPageProtections(DWORD NewProtect)
+	{
+		for (FTransaction* T = Context.GetCurrentTransaction(); T; T = T->Parent)
+		{
+			for (FWriteLogEntry It : T->WriteLog)
+			{
+				ChangePageProtection(It.LogicalAddress, It.Size, /* OldProtect */ 0, NewProtect);
+			}
+		}
+	}
+
+	// Changes the page protection for the virtual memory pages that span
+	// [Address, Address + Size) to NewProtect, if Address is not on the stack.
+	// If OldProtect is not zero, asserts that the previous page protections for
+	// those pages were equal to OldProtect.
+	static void ChangePageProtection(void* Address, size_t Size, DWORD OldProtect, DWORD NewProtect)
+	{
+		if (!FContext::GetThreadStackRange().Contains(Address))
+		{
+			DWORD Prev = 0;
+			BOOL Success = VirtualProtect(Address, Size, NewProtect, &Prev);
+			AUTORTFM_ASSERT(Success);
+			AUTORTFM_ASSERT(OldProtect == 0 || Prev == OldProtect);
+		}
+	}
+
+	// Registers a vectored exception handler that will catch access violation
+	// exceptions and single step exceptions.
+	// Upon catching an access violation, enables the write protection bit for
+	// the virtual memory page that triggered the access violation, and then
+	// requests a single step to execute the write instruction that triggered
+	// the exception. Once that write instruction is complete, a single-step
+	// exception is raised and is caught again by the handler.
+	// Upon catching the single-step exception, re-disables the write protection
+	// bit for the virtual memory page that triggered the access violation and
+	// performs a write log hash, which is compared against the hash snapshotted
+	// when the VEH was originally installed. If the hashes differ, then we've
+	// found the offending write, and we call OnBadWrite().
+	void InstallExceptionHandler()
+	{
+		AUTORTFM_ASSERT(!ExceptionHandler);
+		static FBadWriteIdentifier* Self;
+		Self = this;
+		ExceptionHandler = AddVectoredExceptionHandler(/* CALL_FIRST */ 1, [](_EXCEPTION_POINTERS* Info) -> LONG
+		{
+			static void* PageAddress = nullptr;
+			if (Info->ExceptionRecord->ExceptionCode == EXCEPTION_ACCESS_VIOLATION)
+			{
+				AUTORTFM_ASSERT(PageAddress == nullptr);
+				PageAddress = reinterpret_cast<void*>(Info->ExceptionRecord->ExceptionInformation[1]);
+				
+				// Temporarily unprotect the page so we can step past the write
+				ChangePageProtection(PageAddress, 1, PAGE_READONLY, PAGE_READWRITE);
+
+				// Request a single step
+				Info->ContextRecord->EFlags |= 0x100;
+
+				return EXCEPTION_CONTINUE_EXECUTION;
+			}
+			if (Info->ExceptionRecord->ExceptionCode == EXCEPTION_SINGLE_STEP)
+			{
+				AUTORTFM_ASSERT(PageAddress != nullptr);
+
+				// Re-protect the page
+				ChangePageProtection(PageAddress, 1, PAGE_READWRITE, PAGE_READONLY);
+				PageAddress = nullptr;
+				
+				if (FTransaction* Transaction = FContext::Get()->GetCurrentTransaction())
+				{
+					FWriteHash NewHash = Transaction->CalculateNestedWriteHashWithLimit(Self->NumWriteLogsHashed, nullptr);
+					if (NewHash != Self->Hash)
+					{
+						Self->OnBadWrite();
+					}
+				}
+				return EXCEPTION_CONTINUE_EXECUTION;
+			}
+			return EXCEPTION_CONTINUE_SEARCH;
+		});
+	}
+
+	// Uninstalls the VEH installed by InstallExceptionHandler().
+	void UninstallExceptionHandler()
+	{
+		AUTORTFM_ASSERT(ExceptionHandler);
+		RemoveVectoredExceptionHandler(ExceptionHandler);
+		ExceptionHandler = nullptr;
+	}
+
+	// The transaction transition counter value at the start of the open that
+	// triggered the original memory validation failure.
+	FStateTransitionCounter const OpenCounter;
+
+	// The FContext instance. Stored for easy retrieval.
+	FContext& Context = *FContext::Get();
+
+	// The original memory validation level before being temporarily disabled.
+	EMemoryValidationLevel const OldMemoryValidationLevel = AutoRTFM::ForTheRuntime::GetMemoryValidationLevel();
+
+	// A handle to the installed vectored exception handler.
+	PVOID ExceptionHandler = nullptr;
+
+	// The write log hash snapshot taken before installing the VEH.
+	FWriteHash Hash = 0;
+
+	// The number of write logs hashed to produce Hash.
+	size_t NumWriteLogsHashed = 0;
+};
+#endif
+
+void FTransaction::ValidateWriteHash()
 {
-	const FWriteHash OldHash = RecordedWriteHash;
+	const FWriteHash OldHash = WriteLogSnapshot.Hash;
 	const FWriteHash NewHash = CalculateNestedWriteHash(CurrentOpenReturnAddress);
 
 	if (OldHash != NewHash)
 	{
-		OnValidatorWriteHashMismatch(CurrentMemoryValidationLevel,
+		std::string Message = 
 			"Memory modified in a transaction was also modified in an call to AutoRTFM::Open().\n"
-			"This may lead to memory corruption if the transaction is aborted.");
+			"This may lead to memory corruption if the transaction is aborted.";
+
+		if (AutoRTFM::ForTheRuntime::GetMemoryValidationWriteIdentifierEnabled())
+		{
+#if PLATFORM_WINDOWS
+			if (Context->IsAborting())
+			{
+				// TODO: Support this.
+				Message += "\nWrite identifier is enabled, but cannot currently identify offending write as the context is aborting.";
+			}
+			else
+			{
+				Message += "\nWrite identifier is enabled. Attempting to identify the offending write...";
+				AUTORTFM_WARN("%s", Message.c_str());
+
+				FBadWriteIdentifier* BadWriteFinder = new FBadWriteIdentifier(WriteLogSnapshot.StateTransitionCounter);
+				
+				Context->OnStateTransitionChanged([BadWriteFinder] { BadWriteFinder->OnTransactionStateChange(); });
+
+				// Clear the validation level and write log snapshot to prevent this
+				// getting called again when changing state for the retry. 
+				CurrentMemoryValidationLevel = EMemoryValidationLevel::Disabled;
+				WriteLogSnapshot = {};
+				
+				// Retry the transaction.
+				Context->AbortTransaction(EContextStatus::AbortedByCascadingRetry);
+				// Note: AbortTransaction() throws, so the below is unreachable.
+			}
+#else
+			Message += "\nWrite identifier is enabled, but not supported on this platform.";
+#endif
+		}
+
+		OnValidatorWriteHashMismatch(CurrentMemoryValidationLevel, Message.c_str());
 	}
 }
 
-void FTransaction::OnValidatorWriteHashMismatch(EMemoryValidationLevel MemoryValidationLevel, const char* Message, ...) const
+void FTransaction::OnValidatorWriteHashMismatch(EMemoryValidationLevel MemoryValidationLevel, const char* Message, ...)
 {
 	va_list Args;
 	va_start(Args, Message);
@@ -627,7 +871,7 @@ void FTransaction::OnValidatorWriteHashMismatch(EMemoryValidationLevel MemoryVal
 		case EMemoryValidationLevel::Disabled:
 			break;
 		case EMemoryValidationLevel::Warn:
-			AUTORTFM_WARN_V(Message, Args);
+			LogWithCallstackV(autortfm_log_warn, Message, Args);
 			break;
 		case EMemoryValidationLevel::Error:
 			if (!ForTheRuntime::GetEnsureOnInternalAbort())
@@ -645,7 +889,7 @@ void FTransaction::OnValidatorWriteHashMismatch(EMemoryValidationLevel MemoryVal
 
 FTransaction::FWriteHash FTransaction::CalculateNestedWriteHash(const void *FunctionAddress) const
 {
-	const FWriteHash Hash = CalculateNestedWriteHashWithLimit(NumWriteLogsHashed, FunctionAddress);
+	const FWriteHash Hash = CalculateNestedWriteHashWithLimit(WriteLogSnapshot.NumWriteLogsHashed, FunctionAddress);
 	Context->GetOpenHashThrottler().Update();
 	return Hash;
 }
diff --git a/Engine/Source/Runtime/AutoRTFM/Private/Transaction.h b/Engine/Source/Runtime/AutoRTFM/Private/Transaction.h
--- a/Engine/Source/Runtime/AutoRTFM/Private/Transaction.h
+++ b/Engine/Source/Runtime/AutoRTFM/Private/Transaction.h
@@ -10,6 +10,7 @@
 #include "IntrusivePool.h"
 #include "Stack.h"
 #include "StackRange.h"
+#include "StateTransitionCounter.h"
 #include "Stats.h"
 #include "TaskArray.h"
 #include "Utils.h"
@@ -177,10 +178,10 @@ private:
 	// Compares the hash of all the memory that will be written on abort against the
 	// RecordedWriteHash, raising a warning or error if the hashes are not equal.
 	// See CalculateNestedWriteHash().
-	void ValidateWriteHash() const;
+	void ValidateWriteHash();
 
 	// Warns / errors when a validation hash mismatch is detected.
-	void OnValidatorWriteHashMismatch(EMemoryValidationLevel MemoryValidationLevel, const char* Message, ...) const;
+	static void OnValidatorWriteHashMismatch(EMemoryValidationLevel MemoryValidationLevel, const char* Message, ...);
 	
 	// Returns a hash of all the memory written during this transaction and all parent transactions.
 	// Can be used to ensure that no transaction written memory is modified during an open.
@@ -240,12 +241,25 @@ private:
 	// commits, we need to propagate this to the parent too.
 	TStack<const void*, 1> DeferredPopAllOnAbortHandlers;
 
+	struct FWriteLogSnapshot
+	{
+		FWriteHash Hash = 0;
+		FStateTransitionCounter StateTransitionCounter;
+		size_t NumWriteLogsHashed = 0;
+		
+		inline bool IsZero() const
+		{
+			return Hash == 0 && StateTransitionCounter == FStateTransitionCounter{0} && NumWriteLogsHashed == 0;
+		}
+	};
+
+	struct FBadWriteIdentifier;
+
 	FHitSet HitSet;
 	FIntervalTree NewMemoryTracker;
 	FWriteLog WriteLog;
 	TStatStorage<uint64_t> StatDepth = 1;
-	FWriteHash RecordedWriteHash;
-	size_t NumWriteLogsHashed;
+	FWriteLogSnapshot WriteLogSnapshot;
 	EMemoryValidationLevel CurrentMemoryValidationLevel;
 	const void* CurrentOpenReturnAddress;
 	EState CurrentState;
diff --git a/Engine/Source/Runtime/AutoRTFM/Public/AutoRTFM.h b/Engine/Source/Runtime/AutoRTFM/Public/AutoRTFM.h
--- a/Engine/Source/Runtime/AutoRTFM/Public/AutoRTFM.h
+++ b/Engine/Source/Runtime/AutoRTFM/Public/AutoRTFM.h
@@ -929,6 +929,12 @@ namespace ForTheRuntime
 	// Sets the memory validation throttling mode.
 	UE_AUTORTFM_API void SetMemoryValidationThrottlingEnabled(bool bEnabled);
 
+	// Returns true if the memory validation write-identifier is enabled.
+	UE_AUTORTFM_API bool GetMemoryValidationWriteIdentifierEnabled();
+
+	// Sets true if the memory validation write-identifier mode.
+	UE_AUTORTFM_API void SetMemoryValidationWriteIdentifierEnabled(bool bEnabled);
+
 	// Returns true if the memory validation statistics are enabled.
 	UE_AUTORTFM_API bool GetMemoryValidationStatisticsEnabled();
 
diff --git a/Engine/Source/Runtime/AutoRTFM/Public/AutoRTFM/CAPI.h b/Engine/Source/Runtime/AutoRTFM/Public/AutoRTFM/CAPI.h
--- a/Engine/Source/Runtime/AutoRTFM/Public/AutoRTFM/CAPI.h
+++ b/Engine/Source/Runtime/AutoRTFM/Public/AutoRTFM/CAPI.h
@@ -4,6 +4,7 @@
 
 #include "AutoRTFMDefines.h"
 #include <cstdarg> // for va_list
+#include <cstddef> // for size_t
 
 #ifdef __cplusplus
 extern "C"
@@ -122,6 +123,11 @@ typedef struct
 	// Can be null.
 	void (*OnMemoryValidationThrottlingChanged)();
 
+	// Optional callback to be informed when the value returned by
+	// ForTheRuntime::GetMemoryValidationWriteIdentifierEnabled() changes.
+	// Can be null.
+	void (*OnMemoryValidationWriteIdentifierChanged)();
+
 	// Optional callback to be informed when the value returned by
 	// ForTheRuntime::GetMemoryValidationStatisticsEnabled() changes.
 	// Can be null.
diff --git a/Engine/Source/Runtime/Core/Private/AutoRTFM/AutoRTFMUE.cpp b/Engine/Source/Runtime/Core/Private/AutoRTFM/AutoRTFMUE.cpp
--- a/Engine/Source/Runtime/Core/Private/AutoRTFM/AutoRTFMUE.cpp
+++ b/Engine/Source/Runtime/Core/Private/AutoRTFM/AutoRTFMUE.cpp
@@ -74,6 +74,14 @@ void OnAutoRTFMMemoryValidationThrottlingChanged()
 	return FGenericCrashContext::SetGameData(TEXT("AutoRTFMMemoryValidationThrottlingEnabled"), Text);
 }
 
+void OnAutoRTFMMemoryValidationWriteIdentifierChanged()
+{
+	bool bEnabled = AutoRTFM::ForTheRuntime::GetMemoryValidationWriteIdentifierEnabled();
+	const TCHAR* Text = bEnabled ? TEXT("true") : TEXT("false");
+	UE_LOG(LogAutoRTFM, Log, TEXT("AutoRTFM Memory Validation WriteIdentifier Enabled: %s"), Text);
+	return FGenericCrashContext::SetGameData(TEXT("AutoRTFMMemoryValidationWriteIdentifierEnabled"), Text);
+}
+
 void OnAutoRTFMMemoryValidationStatisticsChanged()
 {
 	bool bEnabled = AutoRTFM::ForTheRuntime::GetMemoryValidationStatisticsEnabled();
@@ -190,6 +198,18 @@ static FAutoConsoleVariable CVarAutoRTFMMemoryValidationThrottlingEnabled(
 	ECVF_Default
 );
 
+static FAutoConsoleVariable CVarAutoRTFMMemoryValidationWriteIdentifierEnabled(
+	TEXT("AutoRTFMMemoryValidationWriteIdentifierEnabled"),
+	AutoRTFM::ForTheRuntime::GetMemoryValidationWriteIdentifierEnabled(),
+	TEXT("Attempts to identify the 'bad' open write on memory validation failure"),
+	FConsoleVariableDelegate::CreateLambda([](IConsoleVariable* Variable)
+	{
+		const bool bValue = Variable->GetBool();
+		AutoRTFM::ForTheRuntime::SetMemoryValidationWriteIdentifierEnabled(bValue);
+	}),
+	ECVF_Default
+);
+
 static FAutoConsoleVariable CVarAutoRTFMMemoryValidationStatisticsEnabled(
 	TEXT("AutoRTFMMemoryValidationStatisticsEnabled"),
 	AutoRTFM::ForTheRuntime::GetMemoryValidationStatisticsEnabled(),
@@ -385,6 +405,7 @@ CORE_API void AutoRTFM::InitializeForUE()
 	ExternAPI.OnRetryTransactionsChanged = OnAutoRTFMRetryTransactionsChanged;
 	ExternAPI.OnMemoryValidationLevelChanged = OnAutoRTFMMemoryValidationLevelChanged;
 	ExternAPI.OnMemoryValidationThrottlingChanged = OnAutoRTFMMemoryValidationThrottlingChanged;
+	ExternAPI.OnMemoryValidationWriteIdentifierChanged = OnAutoRTFMMemoryValidationWriteIdentifierChanged;
 	ExternAPI.OnMemoryValidationStatisticsChanged = OnAutoRTFMMemoryValidationStatisticsChanged;
 
 	AutoRTFM::ForTheRuntime::Initialize(ExternAPI);
@@ -395,6 +416,7 @@ CORE_API void AutoRTFM::InitializeForUE()
 	OnAutoRTFMRetryTransactionsChanged();
 	OnAutoRTFMMemoryValidationLevelChanged();
 	OnAutoRTFMMemoryValidationThrottlingChanged();
+	OnAutoRTFMMemoryValidationWriteIdentifierChanged();
 	OnAutoRTFMMemoryValidationStatisticsChanged();
 }
 
