[AutoRTFM] Fix stale comments.

diff --git a/Engine/Source/Runtime/AutoRTFM/Private/Transaction.h b/Engine/Source/Runtime/AutoRTFM/Private/Transaction.h
--- a/Engine/Source/Runtime/AutoRTFM/Private/Transaction.h
+++ b/Engine/Source/Runtime/AutoRTFM/Private/Transaction.h
@@ -172,11 +172,11 @@ private:
 	void CallAbortTasksWithValidation(FTaskArray& Tasks);
 
 	// Calculates the hash of all the memory written during the transaction and stores it to
-	// RecordedWriteHash.
+	// WriteLogSnapshot.Hash.
 	void RecordWriteHash();
 
 	// Compares the hash of all the memory that will be written on abort against the
-	// RecordedWriteHash, raising a warning or error if the hashes are not equal.
+	// WriteLogSnapshot.Hash, raising a warning or error if the hashes are not equal.
 	// See CalculateNestedWriteHash().
 	void ValidateWriteHash();
 
