[AutoRTFM] Add AUTORTFM_IMPLICIT_ENABLE and AUTORTFM_IMPLICIT_DISABLE

These are new annotations that are applied to various AutoRTFM API parameters so that "naked" lambdas use the appropriate AutoRTFM mode.
This automatically overrides any inherited modes from the function or class, which can often cause confusing errors.

diff --git a/Engine/Source/Runtime/AutoRTFM/Public/AutoRTFM.h b/Engine/Source/Runtime/AutoRTFM/Public/AutoRTFM.h
--- a/Engine/Source/Runtime/AutoRTFM/Public/AutoRTFM.h
+++ b/Engine/Source/Runtime/AutoRTFM/Public/AutoRTFM.h
@@ -10,6 +10,7 @@
 
 // HEADER_UNIT_SKIP - unused warnings
 
+#include "AutoRTFM.h"
 #include "AutoRTFMConstants.h"
 #include "AutoRTFMDefines.h" // IWYU pragma: export
 #include "AutoRTFM/CAPI.h"
@@ -240,7 +241,7 @@ UE_AUTORTFM_CRITICAL_INLINE TransactionID CurrentTransactionID()
 //
 // If AutoRTFM is disabled, the code will be ran non-transactionally.
 template<typename FunctorType>
-UE_AUTORTFM_CRITICAL_INLINE ETransactionResult Transact(const FunctorType& Functor)
+UE_AUTORTFM_CRITICAL_INLINE ETransactionResult Transact(AUTORTFM_IMPLICIT_ENABLE const FunctorType& Functor)
 {
 #if UE_AUTORTFM_ENABLED
 	static constexpr autortfm_mode FunctorAutoRTFMMode = CallMode<FunctorType>;
@@ -276,7 +277,7 @@ UE_AUTORTFM_CRITICAL_INLINE ETransactionResult TransactThenOpen(const FunctorTyp
 // execution if the result is anything other than autortfm_committed.
 // Useful for testing.
 template<typename FunctorType>
-UE_AUTORTFM_CRITICAL_INLINE void Commit(const FunctorType& Functor)
+UE_AUTORTFM_CRITICAL_INLINE void Commit(AUTORTFM_IMPLICIT_ENABLE const FunctorType& Functor)
 {
 #if UE_AUTORTFM_ENABLED
 	static constexpr autortfm_mode FunctorAutoRTFMMode = CallMode<FunctorType>;
@@ -477,7 +478,7 @@ UE_AUTORTFM_CRITICAL_INLINE ReturnType Open(const FunctorType& Functor)
 // attempting to use the transaction is undefined behaviour. The caller should
 // return to the closed as quickly as possible to avoid the risk of the
 // transaction being used in its rolled-back state.
-template<typename FunctorType> [[nodiscard]] UE_AUTORTFM_CRITICAL_INLINE EContextStatus Close(const FunctorType& Functor)
+template<typename FunctorType> [[nodiscard]] UE_AUTORTFM_CRITICAL_INLINE EContextStatus Close(AUTORTFM_IMPLICIT_ENABLE const FunctorType& Functor)
 {
 #if UE_AUTORTFM_ENABLED
 	static constexpr autortfm_mode FunctorAutoRTFMMode = CallMode<FunctorType>;
@@ -514,7 +515,7 @@ UE_AUTORTFM_CRITICAL_INLINE void CascadingRetryTransaction()
 // at that point, the worklist is run in FIFO order.
 // If this is called outside a transaction or from an open nest, then the work
 // happens immediately.
-template<typename FunctorType> UE_AUTORTFM_CRITICAL_INLINE void OnCommit(FunctorType&& Work)
+template<typename FunctorType> UE_AUTORTFM_CRITICAL_INLINE void OnCommit(AUTORTFM_IMPLICIT_DISABLE FunctorType&& Work)
 {
 	if (autortfm_is_closed())
 	{
@@ -533,7 +534,7 @@ template<typename FunctorType> UE_AUTORTFM_CRITICAL_INLINE void OnCommit(Functor
 // Have some work happen when this transaction aborts (before memory rollback).
 // If an abort occurs, the work list is run in LIFO order.
 // If this is called outside a transaction or from an open nest then the work is ignored.
-template<typename FunctorType> UE_AUTORTFM_CRITICAL_INLINE void OnPreAbort(FunctorType&& Work)
+template<typename FunctorType> UE_AUTORTFM_CRITICAL_INLINE void OnPreAbort(AUTORTFM_IMPLICIT_DISABLE FunctorType&& Work)
 {
 	if (autortfm_is_closed())
 	{
@@ -548,7 +549,7 @@ template<typename FunctorType> UE_AUTORTFM_CRITICAL_INLINE void OnPreAbort(Funct
 // Have some work happen when this transaction aborts (after memory rollback).
 // If an abort occurs, the work list is run in LIFO order.
 // If this is called outside a transaction or from an open nest then the work is ignored.
-template<typename FunctorType> UE_AUTORTFM_CRITICAL_INLINE void OnAbort(FunctorType&& Work)
+template<typename FunctorType> UE_AUTORTFM_CRITICAL_INLINE void OnAbort(AUTORTFM_IMPLICIT_DISABLE FunctorType&& Work)
 {
 	if (autortfm_is_closed())
 	{
@@ -565,7 +566,7 @@ template<typename FunctorType> UE_AUTORTFM_CRITICAL_INLINE void OnAbort(FunctorT
 // the end of the transaction.
 // The worklist is run in FIFO order, after the OnCommit or OnAbort worklist has finished.
 // If this is called outside a transaction or from an open nest, then the work is ignored.
-template<typename FunctorType> UE_AUTORTFM_CRITICAL_INLINE void OnComplete(FunctorType&& Work)
+template<typename FunctorType> UE_AUTORTFM_CRITICAL_INLINE void OnComplete(AUTORTFM_IMPLICIT_DISABLE FunctorType&& Work)
 {
 	if (autortfm_is_closed())
 	{
@@ -581,7 +582,7 @@ template<typename FunctorType> UE_AUTORTFM_CRITICAL_INLINE void OnComplete(Funct
 // the handler can be unregistered (see `PopOnCommitHandler`). This is useful
 // for scoped mutations that need an abort handler present unless execution
 // reaches the end of the relevant scope.
-template<typename FunctorType> UE_AUTORTFM_CRITICAL_INLINE void PushOnCommitHandler(const void* Key, FunctorType&& Work)
+template<typename FunctorType> UE_AUTORTFM_CRITICAL_INLINE void PushOnCommitHandler(const void* Key, AUTORTFM_IMPLICIT_DISABLE FunctorType&& Work)
 {
 	if (autortfm_is_closed())
 	{
@@ -593,7 +594,7 @@ template<typename FunctorType> UE_AUTORTFM_CRITICAL_INLINE void PushOnCommitHand
 // the handler can be unregistered (see `PopOnCommitHandler`). This is useful
 // for scoped mutations that need an abort handler present unless execution
 // reaches the end of the relevant scope.
-template<typename FunctorType> UE_AUTORTFM_CRITICAL_INLINE void PushOnCommitHandler(const void*, FunctorType&&) {}
+template<typename FunctorType> UE_AUTORTFM_CRITICAL_INLINE void PushOnCommitHandler(const void*, AUTORTFM_IMPLICIT_DISABLE FunctorType&&) {}
 #endif
 
 #if UE_AUTORTFM
@@ -629,7 +630,7 @@ UE_AUTORTFM_CRITICAL_INLINE void PopAllOnCommitHandlers(const void*) {}
 // the handler can be unregistered (see `PopOnAbortHandler`). This is useful
 // for scoped mutations that need an abort handler present unless execution
 // reaches the end of the relevant scope.
-template<typename FunctorType> UE_AUTORTFM_CRITICAL_INLINE void PushOnAbortHandler(const void* Key, FunctorType&& Work)
+template<typename FunctorType> UE_AUTORTFM_CRITICAL_INLINE void PushOnAbortHandler(const void* Key, AUTORTFM_IMPLICIT_DISABLE FunctorType&& Work)
 {
 	if (autortfm_is_closed())
 	{
diff --git a/Engine/Source/Runtime/AutoRTFM/Public/AutoRTFMDefines.h b/Engine/Source/Runtime/AutoRTFM/Public/AutoRTFMDefines.h
--- a/Engine/Source/Runtime/AutoRTFM/Public/AutoRTFMDefines.h
+++ b/Engine/Source/Runtime/AutoRTFM/Public/AutoRTFMDefines.h
@@ -95,6 +95,16 @@
 // Similar to AUTORTFM_OPEN, but disables memory validation on the call.
 #define AUTORTFM_OPEN_NO_VALIDATION [[clang::autortfm(autortfm_mode_open_no_validation)]]
 
+// Annotation that can be applied to function parameters to implicitly set the
+// AutoRTFM mode of "naked" lambda arguments to AutoRTFM-enabled, regardless of the
+// AutoRTFM mode of the caller function.
+#define AUTORTFM_IMPLICIT_ENABLE [[clang::autortfm_implicit_mode(autortfm_mode_enable)]]
+
+// Annotation that can be applied to function parameters to implicitly set the
+// AutoRTFM mode of "naked" lambda arguments to AutoRTFM-disable, regardless of the
+// AutoRTFM mode of the caller function.
+#define AUTORTFM_IMPLICIT_DISABLE [[clang::autortfm_implicit_mode(autortfm_mode_disable)]]
+
 // Evaluates to the AutoRTFM mode of function, method, type or class.
 #define AUTORTFM_MODE_OF(EXPR_OR_TYPE) static_cast<autortfm_mode>(__autortfm_get_mode(EXPR_OR_TYPE))
 
@@ -117,6 +127,8 @@
 #define AUTORTFM_INFER
 #define AUTORTFM_OPEN
 #define AUTORTFM_OPEN_NO_VALIDATION
+#define AUTORTFM_IMPLICIT_ENABLE
+#define AUTORTFM_IMPLICIT_DISABLE
 #define AUTORTFM_MODE_OF(EXPR_OR_TYPE) autortfm_mode_disable
 #define AUTORTFM_MODE_OF_CALL(EXPR_OR_TYPE) autortfm_mode_disable
 #define UE_AUTORTFM_CALLSITE_FORCEINLINE
diff --git a/Engine/Source/Runtime/AutoRTFM/Public/AutoRTFMTesting.h b/Engine/Source/Runtime/AutoRTFM/Public/AutoRTFMTesting.h
--- a/Engine/Source/Runtime/AutoRTFM/Public/AutoRTFMTesting.h
+++ b/Engine/Source/Runtime/AutoRTFM/Public/AutoRTFMTesting.h
@@ -17,7 +17,7 @@ void AssertionFailure(const char* Expression, const char* File, int Line);
 // execution if the result is anything other than autortfm_committed.
 // Useful for testing.
 template<typename TFunctor>
-static UE_AUTORTFM_FORCEINLINE void Commit(const TFunctor& Functor)
+static UE_AUTORTFM_FORCEINLINE void Commit(AUTORTFM_IMPLICIT_ENABLE const TFunctor& Functor)
 {
 	const ETransactionResult Result = Transact(Functor);
 	AUTORTFM_TESTING_ASSERT(ETransactionResult::Committed == Result);
@@ -26,7 +26,7 @@ static UE_AUTORTFM_FORCEINLINE void Commit(const TFunctor& Functor)
 // Run the callback in a transaction like Transact, but abort program
 // execution if the result is anything other than abort.
 template<typename TFunctor>
-static UE_AUTORTFM_FORCEINLINE void Abort(const TFunctor& Functor)
+static UE_AUTORTFM_FORCEINLINE void Abort(AUTORTFM_IMPLICIT_ENABLE const TFunctor& Functor)
 {
 	const ETransactionResult Result = Transact(Functor);
 	AUTORTFM_TESTING_ASSERT(ETransactionResult::Committed != Result);
