[AutoRTFM] Apply appropriate AutoRTFM modes to lambdas scopes created with the UE_AUTORTFM_XXX macros

Such as UE_AUTORTFM_OPEN, UE_AUTORTFM_ONABORT, etc.

Without this, we can end up with the lambda taking the mode from the outer function, which can cause compiler errors.

diff --git a/Engine/Source/Runtime/AutoRTFM/Public/AutoRTFM.h b/Engine/Source/Runtime/AutoRTFM/Public/AutoRTFM.h
--- a/Engine/Source/Runtime/AutoRTFM/Public/AutoRTFM.h
+++ b/Engine/Source/Runtime/AutoRTFM/Public/AutoRTFM.h
@@ -1080,13 +1080,13 @@ namespace AutoRTFM::Private
 
 #define UE_AUTORTFM_DECLARE_THREAD_LOCAL_VAR_IMPL(Type, Name) Type& Name = ::AutoRTFM::Private::FThreadLocalHelper::Get<Type, __COUNTER__>()
 
-#define UE_AUTORTFM_OPEN_IMPL          ::AutoRTFM::Private::FOpenHelper{} + [&]()
-#define UE_AUTORTFM_OPEN_NO_VALIDATION_IMPL ::AutoRTFM::Private::FOpenNoMemoryValidationHelper{} + [&]()
-#define UE_AUTORTFM_CLOSE_IMPL         ::AutoRTFM::Private::FCloseHelper{} + [&]()
-#define UE_AUTORTFM_ONPREABORT_IMPL(...)  ::AutoRTFM::Private::FOnPreAbortHelper{} + [__VA_ARGS__]() mutable
-#define UE_AUTORTFM_ONABORT_IMPL(...)  ::AutoRTFM::Private::FOnAbortHelper{} + [__VA_ARGS__]() mutable
-#define UE_AUTORTFM_ONCOMMIT_IMPL(...) ::AutoRTFM::Private::FOnCommitHelper{} + [__VA_ARGS__]() mutable
-#define UE_AUTORTFM_TRANSACT_IMPL      ::AutoRTFM::Private::FTransactHelper{} + [&]()
+#define UE_AUTORTFM_OPEN_IMPL          ::AutoRTFM::Private::FOpenHelper{} + [&] AUTORTFM_DISABLE ()
+#define UE_AUTORTFM_OPEN_NO_VALIDATION_IMPL ::AutoRTFM::Private::FOpenNoMemoryValidationHelper{} + [&] AUTORTFM_DISABLE ()
+#define UE_AUTORTFM_CLOSE_IMPL         ::AutoRTFM::Private::FCloseHelper{} + [&] AUTORTFM_ENABLE ()
+#define UE_AUTORTFM_ONPREABORT_IMPL(...)  ::AutoRTFM::Private::FOnPreAbortHelper{} + [__VA_ARGS__] AUTORTFM_DISABLE () mutable
+#define UE_AUTORTFM_ONABORT_IMPL(...)  ::AutoRTFM::Private::FOnAbortHelper{} + [__VA_ARGS__] AUTORTFM_DISABLE () mutable
+#define UE_AUTORTFM_ONCOMMIT_IMPL(...) ::AutoRTFM::Private::FOnCommitHelper{} + [__VA_ARGS__] AUTORTFM_DISABLE () mutable
+#define UE_AUTORTFM_TRANSACT_IMPL      ::AutoRTFM::Private::FTransactHelper{} + [&] AUTORTFM_ENABLE ()
 #else
 
 // Do nothing, these should be followed by blocks that should be either executed or not executed
