[AutoRTFM] Clean up error message for AutoRTFM-enabled calling AutoRTFM-disabled functions & have it point to documentation.

diff --git a/Engine/Source/Runtime/AutoRTFM/Documentation/ModesAndAttributes.md b/Engine/Source/Runtime/AutoRTFM/Documentation/ModesAndAttributes.md
--- a/Engine/Source/Runtime/AutoRTFM/Documentation/ModesAndAttributes.md
+++ b/Engine/Source/Runtime/AutoRTFM/Documentation/ModesAndAttributes.md
@@ -14,7 +14,7 @@ The UBT module flag `bInjectAutoRTFMAttributeDisables` can also be used to apply
 Every function (including lambdas, operators, methods, etc) compiled with AutoRTFM
 instrumentation has an AutoRTFM mode. This is one of the following:
 
-### **AutoRTFM-enabled**
+### AutoRTFM-enabled
 
 AutoRTFM instrumentation **is** generated for the function. This is the default
 mode for targets compiled with AutoRTFM instrumentation enabled.
@@ -23,7 +23,7 @@ An AutoRTFM-enabled function is callable from the **open** and **closed**.
 
 Represented by the enum `autortfm_mode_enable`.
 
-### **AutoRTFM-disabled**
+### AutoRTFM-disabled
 
 AutoRTFM instrumentation **is not** generated for the function.
 This function is callable from the **open**, but not the **closed**.
@@ -39,7 +39,7 @@ assertion if called from the closed.
 
 Represented by the enum `autortfm_mode_disable`.
 
-### **AutoRTFM-inferred**
+### AutoRTFM-inferred
 
 AutoRTFM instrumentation **may be** generated for the function, so long as the
 function does not call any other function that is
@@ -64,27 +64,27 @@ assertions.
 
 Represented by the enum `autortfm_mode_infer`.
 
-### **AutoRTFM-open**
+### AutoRTFM-open
 
 AutoRTFM instrumentation **is not** generated for the function. Instead a stub
 **closed** function is generated that will:
 
 - Transition the transaction state from the **closed** to the **open**.
 - Call the uninstrumented open function.
-- Transition the transaction state back to the **open** before returning.
+- Transition the transaction state back to the **closed** before returning.
 
 This function is callable from the **open** and the **closed**.
 
 Represented by the enum `autortfm_mode_open`.
 
-### **AutoRTFM-open-with-no-validation**
+### AutoRTFM-open-with-no-validation
 
 Behaves the same as [AutoRTFM-open](#autortfm-open), but disables the
 AutoRTFM memory validation for calls to this function.
 
 Represented by the enum `autortfm_mode_open_no_validation`.
 
-### **AutoRTFM-internal**
+### AutoRTFM-internal
 
 Exclusively used by the AutoRTFM runtime. Not to be used externally.
 
@@ -298,34 +298,38 @@ AUTORTFM_DISABLE void DisabledFunction()
 
 ## FAQ
 
-### The compiler errors with `Cannot call AutoRTFM-disabled function 'X' from the AutoRTFM-enabled function 'Y'.`
+### The compiler errors with `error: Cannot call AutoRTFM-disabled function from an AutoRTFM-enabled function`
 
 You are attempting to compile an [AutoRTFM-enabled](#autortfm-enabled) function
-that is making a call to an [AutoRTFM-disabled](#autortfm-disabled)
-(or [AutoRTFM-inferred](#autortfm-inferred) which transitively calls an
-[AutoRTFM-disabled](#autortfm-disabled)) function.
+that is making a call to either:
 
-The compiler is erroring because the closed variant of `Y` has no closed variant
-of `X` to call. To fix this, you have a few choices:
+- An [AutoRTFM-disabled](#autortfm-disabled) function, or
+- An [AutoRTFM-inferred](#autortfm-inferred) function which transitively calls
+  an [AutoRTFM-disabled](#autortfm-disabled) function.
 
-1. If the caller function `Y` is not expected to be reachable in the closed, apply
+The compiler is erroring because the callee function has no closed variant for
+the closed caller function to call.\
+To fix this, you have a few choices:
+
+1. If the caller function is not expected to be reachable in the closed, apply
    [`AUTORTFM_DISABLE`](#autortfm-disabled) to this function (or owning class if
    this is a method and the entire class should not be callable from the closed).
    This option is ideal for reducing code size bloat, improving compilation
    performance, and alerting developers to AutoRTFM mistakes, but
-   `AUTORTFM_DISABLE` is "viral" and may cause callers of function `Y` to start
-   erroring with a similar message until all of those are also fixed up.\
+   `AUTORTFM_DISABLE` is "viral" and may cause callers of the callee function to
+   start erroring with a similar message until all of those are also fixed up.\
    Use of `AUTORTFM_DISABLE` for these errors is the preferred long-term solution.
 2. Alternatively, you may wish to use `AutoRTFM::UnreachableIfClosed()` at the
-   top of function `Y`. This will silence the compiler error, but if function
-   `Y` is called by a closed transaction, then the runtime will assert and
-   crash.\
+   top of the callee function. This will silence the compiler error, but if
+   the callee function is itself called by a closed transaction, then the
+   runtime will assert and crash.\
    This can be used as a short-term solution until `AUTORTFM_DISABLE` can be
    correctly applied.
-3. If function `Y` is a templated function or class that calls a
-   template-dependent function, you may wish to use `AUTORTFM_INFER` which will
-   act similarly to [AUTORTFM_DISABLE](#autortfm-disabled) if and only if the
-   templated function calls an AutoRTFM-disabled function.
+3. If the callee function is a templated function or class that calls a
+   template-dependent function, it may be best to apply `AUTORTFM_INFER` to the
+   callee function. This will act similarly to
+   [AUTORTFM_DISABLE](#autortfm-disabled) if and only if the templated function
+   calls an AutoRTFM-disabled function.
 
 If you're an Epic employee and need assistance, please don't hesitate to request help in
 [#verse-autortfm-dev-ext](https://epic.enterprise.slack.com/archives/C059TRJ6XJ4)
