[AutoRTFM] Move headers (excluding AutoRTFM.h) under AutoRTFM sub-directory.

Changes the includes from `#include "AutoRTFMFEATURE.h"` to `AutoRTFM/FEATURE.h`.

We need to expose some of these limited headers to modules that do not depend on Core.
We currently have a mix of headers in the AutoRTFM module public root directory and some under an AutoRTFM sub-directory.
The latter is more standard across the codebase for includes, so move all other headers to that sub-directory.

diff --git a/Engine/Source/Programs/AutoRTFMTests/Private/API.cpp b/Engine/Source/Programs/AutoRTFMTests/Private/API.cpp
--- a/Engine/Source/Programs/AutoRTFMTests/Private/API.cpp
+++ b/Engine/Source/Programs/AutoRTFMTests/Private/API.cpp
@@ -4,7 +4,7 @@
 
 #include "AutoRTFM.h"
 #include "AutoRTFMTestUtils.h"
-#include "AutoRTFMTesting.h"
+#include "AutoRTFM/Testing.h"
 #include "Catch2Includes.h"
 
 TEST_CASE("API.autortfm_is_transactional")
diff --git a/Engine/Source/Programs/AutoRTFMTests/Private/Abort.cpp b/Engine/Source/Programs/AutoRTFMTests/Private/Abort.cpp
--- a/Engine/Source/Programs/AutoRTFMTests/Private/Abort.cpp
+++ b/Engine/Source/Programs/AutoRTFMTests/Private/Abort.cpp
@@ -1,7 +1,7 @@
 // Copyright Epic Games, Inc. All Rights Reserved.
 
 #include "AutoRTFM.h"
-#include "AutoRTFMTesting.h"
+#include "AutoRTFM/Testing.h"
 #include "AutoRTFMTestUtils.h"
 #include "AutoRTFM/AutoRTFMUE.h"
 #include "Catch2Includes.h"
diff --git a/Engine/Source/Programs/AutoRTFMTests/Private/ArchiveWriterTests.cpp b/Engine/Source/Programs/AutoRTFMTests/Private/ArchiveWriterTests.cpp
--- a/Engine/Source/Programs/AutoRTFMTests/Private/ArchiveWriterTests.cpp
+++ b/Engine/Source/Programs/AutoRTFMTests/Private/ArchiveWriterTests.cpp
@@ -2,7 +2,7 @@
 
 #include "Catch2Includes.h"
 #include "AutoRTFM.h"
-#include "AutoRTFMTesting.h"
+#include "AutoRTFM/Testing.h"
 #include "Misc/NotNull.h"
 #include "Serialization/MemoryWriter.h"
 #include "Serialization/TransactionallySafeArchiveWriter.h"
diff --git a/Engine/Source/Programs/AutoRTFMTests/Private/AutoRTFMTestUtils.h b/Engine/Source/Programs/AutoRTFMTests/Private/AutoRTFMTestUtils.h
--- a/Engine/Source/Programs/AutoRTFMTests/Private/AutoRTFMTestUtils.h
+++ b/Engine/Source/Programs/AutoRTFMTests/Private/AutoRTFMTestUtils.h
@@ -3,7 +3,7 @@
 #pragma once
 
 #include "AutoRTFM.h"
-#include "AutoRTFMDefines.h"
+#include "AutoRTFM/Defines.h"
 #include "Containers/StringConv.h"
 #include "HAL/Platform.h"
 #include "Misc/FeedbackContext.h"
diff --git a/Engine/Source/Programs/AutoRTFMTests/Private/AutoRTFMTests.cpp b/Engine/Source/Programs/AutoRTFMTests/Private/AutoRTFMTests.cpp
--- a/Engine/Source/Programs/AutoRTFMTests/Private/AutoRTFMTests.cpp
+++ b/Engine/Source/Programs/AutoRTFMTests/Private/AutoRTFMTests.cpp
@@ -1,7 +1,7 @@
 // Copyright Epic Games, Inc. All Rights Reserved.
 
 #include "AutoRTFM.h"
-#include "AutoRTFMTesting.h"
+#include "AutoRTFM/Testing.h"
 #include "CoreMinimal.h"
 #include "Misc/EventPool.h"
 #include "Misc/LazySingleton.h"
diff --git a/Engine/Source/Programs/AutoRTFMTests/Private/Benchmarks.cpp b/Engine/Source/Programs/AutoRTFMTests/Private/Benchmarks.cpp
--- a/Engine/Source/Programs/AutoRTFMTests/Private/Benchmarks.cpp
+++ b/Engine/Source/Programs/AutoRTFMTests/Private/Benchmarks.cpp
@@ -1,6 +1,6 @@
 // Copyright Epic Games, Inc. All Rights Reserved.
 
-#include "AutoRTFMTesting.h"
+#include "AutoRTFM/Testing.h"
 #include "Catch2Includes.h"
 
 #include <vector>
diff --git a/Engine/Source/Programs/AutoRTFMTests/Private/CString.cpp b/Engine/Source/Programs/AutoRTFMTests/Private/CString.cpp
--- a/Engine/Source/Programs/AutoRTFMTests/Private/CString.cpp
+++ b/Engine/Source/Programs/AutoRTFMTests/Private/CString.cpp
@@ -1,7 +1,7 @@
 // Copyright Epic Games, Inc. All Rights Reserved.
 
 #include "AutoRTFM.h"
-#include "AutoRTFMTesting.h"
+#include "AutoRTFM/Testing.h"
 #include "Catch2Includes.h"
 #include "CoreGlobals.h"
 #include "Logging/LogVerbosity.h"
diff --git a/Engine/Source/Programs/AutoRTFMTests/Private/Commit.cpp b/Engine/Source/Programs/AutoRTFMTests/Private/Commit.cpp
--- a/Engine/Source/Programs/AutoRTFMTests/Private/Commit.cpp
+++ b/Engine/Source/Programs/AutoRTFMTests/Private/Commit.cpp
@@ -1,6 +1,6 @@
 // Copyright Epic Games, Inc. All Rights Reserved.
 
-#include "AutoRTFMTesting.h"
+#include "AutoRTFM/Testing.h"
 #include "Catch2Includes.h"
 
 #include <map>
diff --git a/Engine/Source/Programs/AutoRTFMTests/Private/CrossCUTests.cpp b/Engine/Source/Programs/AutoRTFMTests/Private/CrossCUTests.cpp
--- a/Engine/Source/Programs/AutoRTFMTests/Private/CrossCUTests.cpp
+++ b/Engine/Source/Programs/AutoRTFMTests/Private/CrossCUTests.cpp
@@ -1,6 +1,6 @@
 // Copyright Epic Games, Inc. All Rights Reserved.
 
-#include "AutoRTFMTesting.h"
+#include "AutoRTFM/Testing.h"
 #include "Catch2Includes.h"
 #include "CoreTypes.h"
 #include "CrossCUTests.h"
diff --git a/Engine/Source/Programs/AutoRTFMTests/Private/CsvProfilerTests.cpp b/Engine/Source/Programs/AutoRTFMTests/Private/CsvProfilerTests.cpp
--- a/Engine/Source/Programs/AutoRTFMTests/Private/CsvProfilerTests.cpp
+++ b/Engine/Source/Programs/AutoRTFMTests/Private/CsvProfilerTests.cpp
@@ -2,7 +2,7 @@
 
 #include "API.h"
 #include "Async/TaskGraphInterfaces.h"
-#include "AutoRTFMTesting.h"
+#include "AutoRTFM/Testing.h"
 #include "Catch2Includes.h"
 #include "CoreTypes.h"
 #include "ProfilingDebugging/CsvProfiler.h"
@@ -182,4 +182,4 @@ TEST_CASE("CsvProfiler.EndCaptureInTransaction")
 	}
 }
 
-#endif // #if 0 // temp fix, also SOL-8472
\ No newline at end of file
+#endif // #if 0 // temp fix, also SOL-8472
diff --git a/Engine/Source/Programs/AutoRTFMTests/Private/Delegate.cpp b/Engine/Source/Programs/AutoRTFMTests/Private/Delegate.cpp
--- a/Engine/Source/Programs/AutoRTFMTests/Private/Delegate.cpp
+++ b/Engine/Source/Programs/AutoRTFMTests/Private/Delegate.cpp
@@ -1,6 +1,6 @@
 // Copyright Epic Games, Inc. All Rights Reserved.
 
-#include "AutoRTFMTesting.h"
+#include "AutoRTFM/Testing.h"
 #include "AutoRTFMTestUtils.h"
 #include "Catch2Includes.h"
 #include "Delegates/Delegate.h"
diff --git a/Engine/Source/Programs/AutoRTFMTests/Private/ExceptionTests.cpp b/Engine/Source/Programs/AutoRTFMTests/Private/ExceptionTests.cpp
--- a/Engine/Source/Programs/AutoRTFMTests/Private/ExceptionTests.cpp
+++ b/Engine/Source/Programs/AutoRTFMTests/Private/ExceptionTests.cpp
@@ -1,6 +1,6 @@
 // Copyright Epic Games, Inc. All Rights Reserved.
 
-#include "AutoRTFMTesting.h"
+#include "AutoRTFM/Testing.h"
 #include "AutoRTFMTestUtils.h"
 #include "Catch2Includes.h"
 #include "BuildMacros.h"
diff --git a/Engine/Source/Programs/AutoRTFMTests/Private/InsightsTests.cpp b/Engine/Source/Programs/AutoRTFMTests/Private/InsightsTests.cpp
--- a/Engine/Source/Programs/AutoRTFMTests/Private/InsightsTests.cpp
+++ b/Engine/Source/Programs/AutoRTFMTests/Private/InsightsTests.cpp
@@ -1,7 +1,7 @@
 // Copyright Epic Games, Inc. All Rights Reserved.
 
 #include "API.h"
-#include "AutoRTFMTesting.h"
+#include "AutoRTFM/Testing.h"
 #include "Catch2Includes.h"
 #include "CoreTypes.h"
 #include "ProfilingDebugging/CpuProfilerTrace.h"
diff --git a/Engine/Source/Programs/AutoRTFMTests/Private/LocalStorageTests.cpp b/Engine/Source/Programs/AutoRTFMTests/Private/LocalStorageTests.cpp
--- a/Engine/Source/Programs/AutoRTFMTests/Private/LocalStorageTests.cpp
+++ b/Engine/Source/Programs/AutoRTFMTests/Private/LocalStorageTests.cpp
@@ -2,7 +2,7 @@
 
 #include "Catch2Includes.h"
 #include "AutoRTFM.h"
-#include "AutoRTFMTesting.h"
+#include "AutoRTFM/Testing.h"
 
 #if PLATFORM_WINDOWS
 
@@ -75,4 +75,4 @@ TEST_CASE("LocalStorage.FLS")
 	FlsFree(FlsSlot);
 }
 
-#endif
\ No newline at end of file
+#endif
diff --git a/Engine/Source/Programs/AutoRTFMTests/Private/LockFreeFixedSizeAllocatorTests.cpp b/Engine/Source/Programs/AutoRTFMTests/Private/LockFreeFixedSizeAllocatorTests.cpp
--- a/Engine/Source/Programs/AutoRTFMTests/Private/LockFreeFixedSizeAllocatorTests.cpp
+++ b/Engine/Source/Programs/AutoRTFMTests/Private/LockFreeFixedSizeAllocatorTests.cpp
@@ -2,7 +2,7 @@
 
 #include "AutoRTFM.h"
 #include "Catch2Includes.h"
-#include "AutoRTFMTesting.h"
+#include "AutoRTFM/Testing.h"
 #include "Containers/LockFreeFixedSizeAllocator.h"
 #include "Containers/LockFreeList.h"
 #include "HAL/ThreadSafeCounter.h"
diff --git a/Engine/Source/Programs/AutoRTFMTests/Private/MTTSAccessDetector.cpp b/Engine/Source/Programs/AutoRTFMTests/Private/MTTSAccessDetector.cpp
--- a/Engine/Source/Programs/AutoRTFMTests/Private/MTTSAccessDetector.cpp
+++ b/Engine/Source/Programs/AutoRTFMTests/Private/MTTSAccessDetector.cpp
@@ -1,6 +1,6 @@
 // Copyright Epic Games, Inc. All Rights Reserved.
 
-#include "AutoRTFMTesting.h"
+#include "AutoRTFM/Testing.h"
 #include "Catch2Includes.h"
 #include "Misc/MTTransactionallySafeAccessDetector.h"
 
diff --git a/Engine/Source/Programs/AutoRTFMTests/Private/MutexTests.cpp b/Engine/Source/Programs/AutoRTFMTests/Private/MutexTests.cpp
--- a/Engine/Source/Programs/AutoRTFMTests/Private/MutexTests.cpp
+++ b/Engine/Source/Programs/AutoRTFMTests/Private/MutexTests.cpp
@@ -1,7 +1,7 @@
 // Copyright Epic Games, Inc. All Rights Reserved.
 
 #include "AutoRTFM.h"
-#include "AutoRTFMTesting.h"
+#include "AutoRTFM/Testing.h"
 #include "AutoRTFMTestUtils.h"
 #include "Catch2Includes.h"
 #include "Async/TransactionallySafeMutex.h"
diff --git a/Engine/Source/Programs/AutoRTFMTests/Private/OpenWrapperTests.cpp b/Engine/Source/Programs/AutoRTFMTests/Private/OpenWrapperTests.cpp
--- a/Engine/Source/Programs/AutoRTFMTests/Private/OpenWrapperTests.cpp
+++ b/Engine/Source/Programs/AutoRTFMTests/Private/OpenWrapperTests.cpp
@@ -2,7 +2,7 @@
 
 #include "AutoRTFM.h"
 #include "AutoRTFM/OpenWrapper.h"
-#include "AutoRTFMTesting.h"
+#include "AutoRTFM/Testing.h"
 #include "Catch2Includes.h"
 
 TEST_CASE("OpenWrapper")
diff --git a/Engine/Source/Programs/AutoRTFMTests/Private/QueueTests.cpp b/Engine/Source/Programs/AutoRTFMTests/Private/QueueTests.cpp
--- a/Engine/Source/Programs/AutoRTFMTests/Private/QueueTests.cpp
+++ b/Engine/Source/Programs/AutoRTFMTests/Private/QueueTests.cpp
@@ -1,6 +1,6 @@
 // Copyright Epic Games, Inc. All Rights Reserved.
 
-#include "AutoRTFMTesting.h"
+#include "AutoRTFM/Testing.h"
 #include "Catch2Includes.h"
 #include "Containers/TransactionallySafeSpscQueue.h"
 
diff --git a/Engine/Source/Programs/AutoRTFMTests/Private/RWLockTests.cpp b/Engine/Source/Programs/AutoRTFMTests/Private/RWLockTests.cpp
--- a/Engine/Source/Programs/AutoRTFMTests/Private/RWLockTests.cpp
+++ b/Engine/Source/Programs/AutoRTFMTests/Private/RWLockTests.cpp
@@ -2,7 +2,7 @@
 
 #include "Catch2Includes.h"
 #include "AutoRTFM.h"
-#include "AutoRTFMTesting.h"
+#include "AutoRTFM/Testing.h"
 #include "AutoRTFMTestUtils.h"
 #include "Misc/TransactionallySafeRWLock.h"
 #include "Misc/ScopeRWLock.h"
diff --git a/Engine/Source/Programs/AutoRTFMTests/Private/RefCountPtrTests.cpp b/Engine/Source/Programs/AutoRTFMTests/Private/RefCountPtrTests.cpp
--- a/Engine/Source/Programs/AutoRTFMTests/Private/RefCountPtrTests.cpp
+++ b/Engine/Source/Programs/AutoRTFMTests/Private/RefCountPtrTests.cpp
@@ -1,7 +1,7 @@
 // Copyright Epic Games, Inc. All Rights Reserved.
 
 #include "AutoRTFM.h"
-#include "AutoRTFMTesting.h"
+#include "AutoRTFM/Testing.h"
 #include "Catch2Includes.h"
 #include "Templates/RefCounting.h"
 
diff --git a/Engine/Source/Programs/AutoRTFMTests/Private/Statics.cpp b/Engine/Source/Programs/AutoRTFMTests/Private/Statics.cpp
--- a/Engine/Source/Programs/AutoRTFMTests/Private/Statics.cpp
+++ b/Engine/Source/Programs/AutoRTFMTests/Private/Statics.cpp
@@ -1,8 +1,8 @@
 // Copyright Epic Games, Inc. All Rights Reserved.
 
 #include "AutoRTFM.h"
-#include "AutoRTFMDefines.h"
-#include "AutoRTFMTesting.h"
+#include "AutoRTFM/Defines.h"
+#include "AutoRTFM/Testing.h"
 #include "Context.h"
 #include "ContextInlines.h"
 #include "Transaction.h"
diff --git a/Engine/Source/Programs/AutoRTFMTests/Private/StatsTests.cpp b/Engine/Source/Programs/AutoRTFMTests/Private/StatsTests.cpp
--- a/Engine/Source/Programs/AutoRTFMTests/Private/StatsTests.cpp
+++ b/Engine/Source/Programs/AutoRTFMTests/Private/StatsTests.cpp
@@ -2,7 +2,7 @@
 
 #include "Catch2Includes.h"
 #include "AutoRTFM.h"
-#include "AutoRTFMTesting.h"
+#include "AutoRTFM/Testing.h"
 
 #include "Async/TaskGraphInterfaces.h"
 #include "Stats/Stats.h"
diff --git a/Engine/Source/Programs/AutoRTFMTests/Private/TaskTests.cpp b/Engine/Source/Programs/AutoRTFMTests/Private/TaskTests.cpp
--- a/Engine/Source/Programs/AutoRTFMTests/Private/TaskTests.cpp
+++ b/Engine/Source/Programs/AutoRTFMTests/Private/TaskTests.cpp
@@ -1,7 +1,6 @@
 // Copyright Epic Games, Inc. All Rights Reserved.
 
 #include "AutoRTFM.h"
-#include "AutoRTFMTask.h"
 #include "Catch2Includes.h"
 #include "GenericPlatform/GenericPlatformMisc.h"
 #include "Templates/UnrealTemplate.h"
diff --git a/Engine/Source/Programs/AutoRTFMTests/Private/Tests.cpp b/Engine/Source/Programs/AutoRTFMTests/Private/Tests.cpp
--- a/Engine/Source/Programs/AutoRTFMTests/Private/Tests.cpp
+++ b/Engine/Source/Programs/AutoRTFMTests/Private/Tests.cpp
@@ -1,7 +1,7 @@
 // Copyright Epic Games, Inc. All Rights Reserved.
 
 #include "AutoRTFM.h"
-#include "AutoRTFMTesting.h"
+#include "AutoRTFM/Testing.h"
 #include "Catch2Includes.h"
 #include "Logging/LogMacros.h"
 
diff --git a/Engine/Source/Programs/AutoRTFMTests/Private/UECoreTests.cpp b/Engine/Source/Programs/AutoRTFMTests/Private/UECoreTests.cpp
--- a/Engine/Source/Programs/AutoRTFMTests/Private/UECoreTests.cpp
+++ b/Engine/Source/Programs/AutoRTFMTests/Private/UECoreTests.cpp
@@ -6,7 +6,7 @@
 #include "Async/TransactionallySafeMutex.h"
 #include "Async/ParallelFor.h"
 #include "AssetRegistry/AssetDataTagMap.h"
-#include "AutoRTFMTesting.h"
+#include "AutoRTFM/Testing.h"
 #include "Blueprint/BlueprintExceptionInfo.h"
 #include "Delegates/IDelegateInstance.h"
 #include "Experimental/UnifiedError/CoreErrorTypes.h"
diff --git a/Engine/Source/Programs/AutoRTFMTests/Private/UEPlatform.cpp b/Engine/Source/Programs/AutoRTFMTests/Private/UEPlatform.cpp
--- a/Engine/Source/Programs/AutoRTFMTests/Private/UEPlatform.cpp
+++ b/Engine/Source/Programs/AutoRTFMTests/Private/UEPlatform.cpp
@@ -1,6 +1,6 @@
 // Copyright Epic Games, Inc. All Rights Reserved.
 
-#include "AutoRTFMTesting.h"
+#include "AutoRTFM/Testing.h"
 #include "AutoRTFM.h"
 #include "AutoRTFMTestUtils.h"
 #include "Catch2Includes.h"
diff --git a/Engine/Source/Programs/AutoRTFMTests/Private/UObjectTests.cpp b/Engine/Source/Programs/AutoRTFMTests/Private/UObjectTests.cpp
--- a/Engine/Source/Programs/AutoRTFMTests/Private/UObjectTests.cpp
+++ b/Engine/Source/Programs/AutoRTFMTests/Private/UObjectTests.cpp
@@ -1,7 +1,7 @@
 // Copyright Epic Games, Inc. All Rights Reserved.
 
 #include "AutoRTFM.h"
-#include "AutoRTFMTesting.h"
+#include "AutoRTFM/Testing.h"
 #include "ScopedGuard.h"
 #include "GenericPlatform/GenericPlatformMisc.h"
 #include "HAL/MallocLeakDetection.h"
diff --git a/Engine/Source/Programs/AutoRTFMTests/Private/UPackageTests.cpp b/Engine/Source/Programs/AutoRTFMTests/Private/UPackageTests.cpp
--- a/Engine/Source/Programs/AutoRTFMTests/Private/UPackageTests.cpp
+++ b/Engine/Source/Programs/AutoRTFMTests/Private/UPackageTests.cpp
@@ -2,7 +2,7 @@
 
 #include "AutoRTFM.h"
 #include "AutoRTFMTestUtils.h"
-#include "AutoRTFMTesting.h"
+#include "AutoRTFM/Testing.h"
 #include "MyAutoRTFMTestObject.h"
 #include "Misc/PackageName.h"
 #include "UObject/LinkerInstancingContext.h"
diff --git a/Engine/Source/Runtime/AutoRTFM/Private/API.cpp b/Engine/Source/Runtime/AutoRTFM/Private/API.cpp
--- a/Engine/Source/Runtime/AutoRTFM/Private/API.cpp
+++ b/Engine/Source/Runtime/AutoRTFM/Private/API.cpp
@@ -1,7 +1,7 @@
 // Copyright Epic Games, Inc. All Rights Reserved.
 
 #include "AutoRTFM.h"
-#include "AutoRTFMDefines.h"
+#include "AutoRTFM/Defines.h"
 #include "Context.h"
 #include "ExternAPI.h"
 #include "Transaction.h"
@@ -451,7 +451,7 @@ namespace AutoRTFM
 }
 
 #if (defined(__AUTORTFM_ENABLED) && __AUTORTFM_ENABLED)
-#include "AutoRTFMConstants.h"
+#include "AutoRTFM/Constants.h"
 #include "CallNest.h"
 #include "Context.h"
 #include "ContextInlines.h"
diff --git a/Engine/Source/Runtime/AutoRTFM/Private/CompilerInterface.cpp b/Engine/Source/Runtime/AutoRTFM/Private/CompilerInterface.cpp
--- a/Engine/Source/Runtime/AutoRTFM/Private/CompilerInterface.cpp
+++ b/Engine/Source/Runtime/AutoRTFM/Private/CompilerInterface.cpp
@@ -1,6 +1,6 @@
 // Copyright Epic Games, Inc. All Rights Reserved.
 
-#include "AutoRTFMDefines.h"
+#include "AutoRTFM/Defines.h"
 #if (defined(__AUTORTFM) && __AUTORTFM)
 
 #include "AutoRTFM.h"
diff --git a/Engine/Source/Runtime/AutoRTFM/Private/Context.cpp b/Engine/Source/Runtime/AutoRTFM/Private/Context.cpp
--- a/Engine/Source/Runtime/AutoRTFM/Private/Context.cpp
+++ b/Engine/Source/Runtime/AutoRTFM/Private/Context.cpp
@@ -5,7 +5,7 @@
 #include "Context.h"
 #include "ContextInlines.h"
 
-#include "AutoRTFMMetrics.h"
+#include "AutoRTFM/Metrics.h"
 #include "CallNestInlines.h"
 #include "ExternAPI.h"
 #include "FunctionMap.h"
diff --git a/Engine/Source/Runtime/AutoRTFM/Private/Context.h b/Engine/Source/Runtime/AutoRTFM/Private/Context.h
--- a/Engine/Source/Runtime/AutoRTFM/Private/Context.h
+++ b/Engine/Source/Runtime/AutoRTFM/Private/Context.h
@@ -4,7 +4,7 @@
 
 #if (defined(__AUTORTFM) && __AUTORTFM)
 
-#include "AutoRTFMTask.h"
+#include "AutoRTFM/Task.h"
 #include "CallNest.h"
 #include "IntrusivePool.h"
 #include "OpenHashThrottler.h"
diff --git a/Engine/Source/Runtime/AutoRTFM/Private/FunctionMapInlines.h b/Engine/Source/Runtime/AutoRTFM/Private/FunctionMapInlines.h
--- a/Engine/Source/Runtime/AutoRTFM/Private/FunctionMapInlines.h
+++ b/Engine/Source/Runtime/AutoRTFM/Private/FunctionMapInlines.h
@@ -4,8 +4,8 @@
 
 #if (defined(__AUTORTFM) && __AUTORTFM)
 
-#include "AutoRTFMDefines.h"
-#include "AutoRTFMConstants.h"
+#include "AutoRTFM/Defines.h"
+#include "AutoRTFM/Constants.h"
 #include "FunctionMap.h"
 #include "Utils.h"
 
diff --git a/Engine/Source/Runtime/AutoRTFM/Private/IntrusivePool.h b/Engine/Source/Runtime/AutoRTFM/Private/IntrusivePool.h
--- a/Engine/Source/Runtime/AutoRTFM/Private/IntrusivePool.h
+++ b/Engine/Source/Runtime/AutoRTFM/Private/IntrusivePool.h
@@ -4,7 +4,7 @@
 
 #if (defined(__AUTORTFM) && __AUTORTFM)
 
-#include "AutoRTFMDefines.h"
+#include "AutoRTFM/Defines.h"
 #include "BlockAllocator.h"
 
 #include <cstddef>
diff --git a/Engine/Source/Runtime/AutoRTFM/Private/Memcpy.h b/Engine/Source/Runtime/AutoRTFM/Private/Memcpy.h
--- a/Engine/Source/Runtime/AutoRTFM/Private/Memcpy.h
+++ b/Engine/Source/Runtime/AutoRTFM/Private/Memcpy.h
@@ -4,7 +4,7 @@
 
 #if (defined(__AUTORTFM) && __AUTORTFM)
 
-#include "AutoRTFMDefines.h"
+#include "AutoRTFM/Defines.h"
 
 #include <stdlib.h>
 
diff --git a/Engine/Source/Runtime/AutoRTFM/Private/Utils.h b/Engine/Source/Runtime/AutoRTFM/Private/Utils.h
--- a/Engine/Source/Runtime/AutoRTFM/Private/Utils.h
+++ b/Engine/Source/Runtime/AutoRTFM/Private/Utils.h
@@ -2,7 +2,7 @@
 
 #pragma once
 
-#include "AutoRTFMDefines.h"
+#include "AutoRTFM/Defines.h"
 #if (defined(__AUTORTFM) && __AUTORTFM)
 
 #include "AutoRTFM.h"
diff --git a/Engine/Source/Runtime/AutoRTFM/Public/AutoRTFM.h b/Engine/Source/Runtime/AutoRTFM/Public/AutoRTFM.h
--- a/Engine/Source/Runtime/AutoRTFM/Public/AutoRTFM.h
+++ b/Engine/Source/Runtime/AutoRTFM/Public/AutoRTFM.h
@@ -10,14 +10,13 @@
 
 // HEADER_UNIT_SKIP - unused warnings
 
-#include "AutoRTFM.h"
-#include "AutoRTFMConstants.h"
-#include "AutoRTFMDefines.h" // IWYU pragma: export
+#include "AutoRTFM/Constants.h"
+#include "AutoRTFM/Defines.h" // IWYU pragma: export
 #include "AutoRTFM/CAPI.h"
 
 #ifdef __cplusplus
 
-#include "AutoRTFMTask.h"
+#include "AutoRTFM/Task.h"
 
 #include <cstdarg>
 #include <new>
diff --git a/Engine/Source/Runtime/AutoRTFM/Public/AutoRTFM/CAPI.h b/Engine/Source/Runtime/AutoRTFM/Public/AutoRTFM/CAPI.h
--- a/Engine/Source/Runtime/AutoRTFM/Public/AutoRTFM/CAPI.h
+++ b/Engine/Source/Runtime/AutoRTFM/Public/AutoRTFM/CAPI.h
@@ -2,7 +2,8 @@
 
 #pragma once
 
-#include "AutoRTFMDefines.h"
+#include "AutoRTFM/Defines.h"
+
 #include <cstdarg> // for va_list
 #include <cstddef> // for size_t
 
diff --git a/Engine/Source/Runtime/AutoRTFM/Public/AutoRTFM/Constants.h b/Engine/Source/Runtime/AutoRTFM/Public/AutoRTFM/Constants.h
new file mode 100644
--- /dev/null
+++ b/Engine/Source/Runtime/AutoRTFM/Public/AutoRTFM/Constants.h
@@ -0,0 +1,121 @@
+// Copyright Epic Games, Inc. All Rights Reserved.
+
+#pragma once
+
+// WARNING: Any change in these constants will require a re-patch and re-build of LLVM!
+
+#ifdef __cplusplus
+#include <cstdint>
+
+extern "C"
+{
+#endif
+
+// An enumerator of transactional memory validation levels.
+// Memory validation is used to detect modification by open-code to memory that was written by a
+// transaction. In this situation, aborting the transaction can corrupt memory as the undo will
+// overwrite the writes made in the open-code.
+typedef enum
+{
+	// Use the default memory validation level.
+	autortfm_memory_validation_level_default,
+
+	// Disable memory validation.
+	autortfm_memory_validation_level_disabled,
+
+	// Enable memory validation. Memory validation failures are treated as warnings.
+	autortfm_memory_validation_level_warn,
+
+	// Enable memory validation. Memory validation are treated as fatal.
+	autortfm_memory_validation_level_error,
+} autortfm_memory_validation_level;
+
+// An enumerator of mutually exclusive AutoRTFM modes applied to functions or classes.
+typedef enum
+{
+	// [[clang::autortfm(enable)]] / default
+	// Generate AutoRTFM instrumentation for the function.
+	autortfm_mode_enable = 0,
+
+	// [[clang::autortfm(disable)]]
+	// Do not generate AutoRTFM instrumentation for the function.
+	// Calls to the function from AutoRTFM-enabled functions is a compile-time error.
+	// WARNING: If you change this constant, then make sure to also update 
+	//          UEBuildModule.AddModuleToCompileEnvironment() in
+	//          Engine\Source\Programs\UnrealBuildTool\Configuration\UEBuildModule.cs
+	autortfm_mode_disable = 1,
+
+	// [[clang::autortfm(infer)]]
+	// AutoRTFM instrumentation mode is determined by the calls made by the function.
+	autortfm_mode_infer = 2,
+
+	// [[clang::autortfm(open)]] / [[clang::autortfm_always_open]]
+	// Generate AutoRTFM instrumentation for the function.
+	// Calling the function from the closed will call the uninstrumented variant.
+	autortfm_mode_open = 3,
+
+	// [[clang::autortfm(open_no_validation)]] / [[clang::autortfm_always_open_disable_memory_validation]]
+	// Open, but memory validation will not be performed.
+	autortfm_mode_open_no_validation = 4,
+
+	// [[clang::autortfm(autortfm_mode_internal)]]
+	// Do not generate AutoRTFM instrumentation for the function.
+	// Calls to the function from AutoRTFM-enabled functions will directly call
+	// the open function without any pre / post callbacks or trampolines.
+	autortfm_mode_internal = 5,
+
+	// [[clang::noautortfm]]
+	// DEPRECATED - similar to Disable but does not cause a compile-time error if
+	// called from an AutoRTFM-enabled function.
+	autortfm_mode_no_autortfm = 6,
+} autortfm_mode;
+
+// The last mode declared in autortfm_mode. Used for range checking.
+static const autortfm_mode autortfm_last_mode = autortfm_mode_no_autortfm;
+
+// Returns the AutoRTFM mode as a string
+inline const char* autortfm_mode_to_string(autortfm_mode Mode) {
+  switch (Mode) {
+    case autortfm_mode_enable:
+      return "autortfm_mode_enable";
+    case autortfm_mode_disable:
+      return "autortfm_mode_disable";
+    case autortfm_mode_infer:
+      return "autortfm_mode_infer";
+    case autortfm_mode_open:
+      return "autortfm_mode_open";
+    case autortfm_mode_open_no_validation:
+      return "autortfm_mode_open_no_validation";
+    case autortfm_mode_internal:
+      return "autortfm_mode_internal";
+    case autortfm_mode_no_autortfm:
+      return "autortfm_mode_no_autortfm";
+  }
+}
+
+#ifdef __cplusplus
+}
+#endif
+
+#ifdef __cplusplus
+namespace AutoRTFM::Constants
+{
+	inline constexpr uint32_t Major = 0;
+	inline constexpr uint32_t Minor = 2;
+	inline constexpr uint32_t Patch = 0;
+
+	// The Magic Prefix constant - an arbitrarily chosen address prefix, shared
+	// between the AutoRTFM compiler and runtime.
+	// We add this prefix value to open function pointer addresses in our custom 
+	// LLVM pass. At runtime, if we detect the Magic Prefix in the the top 16 bits
+	// of an open function pointer address, we assume that we can find a closed
+	// variant pointer residing 8 bytes before the function address.
+	inline constexpr uint64_t MagicPrefix = 0xa273'0000'0000'0000;
+	// Similar to the above, but these constants indicate that the low 48 bits
+	// provide a relative addresses from the open function to the closed
+	// function.
+	inline constexpr uint64_t PosOffsetMagicPrefix = 0xa272'0000'0000'0000;
+	inline constexpr uint64_t NegOffsetMagicPrefix = 0xa271'0000'0000'0000;
+
+} // namespace AutoRTFM::Constants
+#endif
diff --git a/Engine/Source/Runtime/AutoRTFM/Public/AutoRTFM/Defines.h b/Engine/Source/Runtime/AutoRTFM/Public/AutoRTFM/Defines.h
new file mode 100644
--- /dev/null
+++ b/Engine/Source/Runtime/AutoRTFM/Public/AutoRTFM/Defines.h
@@ -0,0 +1,215 @@
+// Copyright Epic Games, Inc. All Rights Reserved.
+
+// IWYU pragma: private
+// Include include "AutoRTFM.h" instead
+
+#pragma once
+
+#include "AutoRTFM/Constants.h"
+
+#if (defined(__AUTORTFM) && __AUTORTFM)
+#define UE_AUTORTFM 1  // Compiler is 'verse-clang'
+#else
+#define UE_AUTORTFM 0
+#endif
+
+#if (defined(__AUTORTFM_ENABLED) && __AUTORTFM_ENABLED)
+#define UE_AUTORTFM_ENABLED 1  // Compiled with '-fautortfm'
+#else
+#define UE_AUTORTFM_ENABLED 0
+#endif
+
+#if !defined(UE_AUTORTFM_ENABLED_RUNTIME_BY_DEFAULT)
+#define UE_AUTORTFM_ENABLED_RUNTIME_BY_DEFAULT 1
+#endif
+
+#if !defined(UE_AUTORTFM_STATIC_VERIFIER)
+#define UE_AUTORTFM_STATIC_VERIFIER 0
+#endif
+
+#if UE_AUTORTFM
+// The annotated function will have no AutoRTFM closed variant generated, and
+// cannot be called from another closed function. This attribute will eventually
+// be deprecated and replaced with AUTORTFM_DISABLE.
+#define UE_AUTORTFM_NOAUTORTFM [[clang::noautortfm]]
+
+// Omits AutoRTFM instrumentation from the function. Calling this annotated
+// function from the closed will automatically enter the open for the duration
+// of the call and will return back to the closed.
+#define UE_AUTORTFM_ALWAYS_OPEN [[clang::autortfm_always_open]]
+
+// The same as UE_AUTORTFM_ALWAYS_OPEN, but disables memory validation for this
+// call.
+#define UE_AUTORTFM_ALWAYS_OPEN_NO_MEMORY_VALIDATION [[clang::autortfm_always_open_disable_memory_validation]]
+
+// Annotation that can be applied to classes, methods or functions to prevent
+// AutoRTFM closed function(s) from being generated.
+// Applying the annotation to a class is equivalent to applying the annotation
+// to each method of the class.
+// When annotating a function, the attribute must be placed on a the function
+// declaration (usually in the header) and not a function implementation.
+// Annotated functions cannot be called from another closed function and
+// attempting to do so will result in a runtime failure.
+#define AUTORTFM_DISABLE [[clang::autortfm(autortfm_mode_disable)]]
+
+// Applies the AUTORTFM_DISABLE annotation if the condition argument evaluates to true.
+// Warning: This is an experimental API and may be removed in the future.
+#define AUTORTFM_DISABLE_IF(...) [[clang::autortfm(autortfm_mode_disable, __VA_ARGS__)]]
+
+// Begins a range where all classes and functions will be automatically
+// annotated with AUTORTFM_DISABLE. Must be ended with AUTORTFM_DISABLE_END
+// before the end of the file.
+#define AUTORTFM_DISABLE_BEGIN _Pragma("clang attribute AutoRTFM_Disable.push (AUTORTFM_DISABLE, apply_to = any(function, record))")
+
+// Ends a range started with AUTORTFM_DISABLE_BEGIN
+#define AUTORTFM_DISABLE_END _Pragma("clang attribute AutoRTFM_Disable.pop")
+
+// Annotation that can be applied to classes, methods or functions to re-enable
+// AutoRTFM instrumentation which would otherwise be disabled by AUTORTFM_DISABLE.
+// Useful for selectively enabling AutoRTFM on methods when a class is annotated
+// AUTORTFM_DISABLE.
+#define AUTORTFM_ENABLE [[clang::autortfm(autortfm_mode_enable)]]
+
+// Annotation that can be applied to classes, methods or functions to infer
+// whether AutoRTFM instrumentation should be enabled for each individual
+// function based on the AutoRTFM-enabled state of each callee made by the
+// function. If the function calls any AutoRTFM-disabled function, then the
+// function will also be AutoRTFM-disabled, otherwise the function is 
+// AutoRTFM-enabled.
+// Applying the annotation to a class is equivalent to applying the annotation
+// to each method of the class.
+// When annotating a function, the attribute must be placed on a the function
+// declaration (usually in the header) and not a function implementation.
+#define AUTORTFM_INFER [[clang::autortfm(autortfm_mode_infer)]]
+
+// Annotation that can be applied to classes, methods or functions to prevent
+// AutoRTFM closed function(s) from being generated. Unlike AUTORTFM_DISABLE
+// annotated functions can be called from closed functions, which will call
+// the uninstrumented function.
+// Applying the annotation to a class is equivalent to applying the annotation
+// to each method of the class.
+// When annotating a function, the attribute must be placed on a the function
+// declaration (usually in the header) and not a function implementation.
+#define AUTORTFM_OPEN [[clang::autortfm(autortfm_mode_open)]]
+
+// Similar to AUTORTFM_OPEN, but disables memory validation on the call.
+#define AUTORTFM_OPEN_NO_VALIDATION [[clang::autortfm(autortfm_mode_open_no_validation)]]
+
+// Annotation that can be applied to function parameters to implicitly set the
+// AutoRTFM mode of "naked" lambda arguments to AutoRTFM-enabled, regardless of the
+// AutoRTFM mode of the caller function.
+#define AUTORTFM_IMPLICIT_ENABLE [[clang::autortfm_implicit_mode(autortfm_mode_enable)]]
+
+// Annotation that can be applied to function parameters to implicitly set the
+// AutoRTFM mode of "naked" lambda arguments to AutoRTFM-disable, regardless of the
+// AutoRTFM mode of the caller function.
+#define AUTORTFM_IMPLICIT_DISABLE [[clang::autortfm_implicit_mode(autortfm_mode_disable)]]
+
+// Evaluates to the AutoRTFM mode of function, method, type or class.
+#define AUTORTFM_MODE_OF(EXPR_OR_TYPE) static_cast<autortfm_mode>(__autortfm_get_mode(EXPR_OR_TYPE))
+
+// Evaluates to the AutoRTFM mode of function, method, type or class.
+#define AUTORTFM_MODE_OF_CALL(CALL_EXPR) AUTORTFM_MODE_OF(__autortfm_declcall(CALL_EXPR))
+
+// Force the call statement to be inlined.
+#define UE_AUTORTFM_CALLSITE_FORCEINLINE [[clang::always_inline]]
+
+#else // ^^^ UE_AUTORTFM ^^^ | vvv !UE_AUTORTFM vvv
+
+#define UE_AUTORTFM_NOAUTORTFM
+#define UE_AUTORTFM_ALWAYS_OPEN
+#define UE_AUTORTFM_ALWAYS_OPEN_NO_MEMORY_VALIDATION
+#define AUTORTFM_DISABLE
+#define AUTORTFM_DISABLE_IF(CONDITION)
+#define AUTORTFM_DISABLE_BEGIN
+#define AUTORTFM_DISABLE_END
+#define AUTORTFM_ENABLE
+#define AUTORTFM_INFER
+#define AUTORTFM_OPEN
+#define AUTORTFM_OPEN_NO_VALIDATION
+#define AUTORTFM_IMPLICIT_ENABLE
+#define AUTORTFM_IMPLICIT_DISABLE
+#define AUTORTFM_MODE_OF(EXPR_OR_TYPE) autortfm_mode_disable
+#define AUTORTFM_MODE_OF_CALL(EXPR_OR_TYPE) autortfm_mode_disable
+#define UE_AUTORTFM_CALLSITE_FORCEINLINE
+
+#endif
+
+#ifdef __cplusplus
+#define AUTORTFM_NOEXCEPT noexcept
+#define AUTORTFM_EXCEPT noexcept(false)
+#else
+#define AUTORTFM_NOEXCEPT
+#define AUTORTFM_EXCEPT
+#endif
+
+#if UE_AUTORTFM && UE_AUTORTFM_STATIC_VERIFIER
+#define UE_AUTORTFM_ENSURE_SAFE [[clang::autortfm_ensure_safe]]
+#define UE_AUTORTFM_ASSUME_SAFE [[clang::autortfm_assume_safe]]
+#else
+#define UE_AUTORTFM_ENSURE_SAFE
+#define UE_AUTORTFM_ASSUME_SAFE
+#endif
+
+#define UE_AUTORTFM_CONCAT_IMPL(A, B) A ## B
+#define UE_AUTORTFM_CONCAT(A, B) UE_AUTORTFM_CONCAT_IMPL(A, B)
+
+#if defined(__cplusplus) && defined(__UNREAL__) && !(defined(UE_AUTORTFM_DO_NOT_INCLUDE_PLATFORM_H) && UE_AUTORTFM_DO_NOT_INCLUDE_PLATFORM_H)
+	// Include HAL/Platform.h for DLLIMPORT / DLLEXPORT definitions, which
+	// UBT can use as a definition for AUTORTFM_API.
+	#include <HAL/Platform.h>
+	#define UE_AUTORTFM_API AUTORTFM_API
+#else
+	#ifndef UE_AUTORTFM_API
+		#define UE_AUTORTFM_API
+	#endif
+#endif
+
+#if defined(_MSC_VER)
+#define UE_AUTORTFM_FORCEINLINE __forceinline
+#define UE_AUTORTFM_FORCEINLINE_ALWAYS __forceinline
+#define UE_AUTORTFM_FORCENOINLINE __declspec(noinline)
+#define UE_AUTORTFM_ASSUME(x) __assume(x)
+#elif defined(__clang__)
+#define UE_AUTORTFM_FORCEINLINE __attribute__((always_inline)) inline
+#define UE_AUTORTFM_FORCEINLINE_ALWAYS __attribute__((always_inline)) inline
+#define UE_AUTORTFM_FORCENOINLINE __attribute__((noinline))
+#define UE_AUTORTFM_ASSUME(x) __builtin_assume(x)
+#else
+#define UE_AUTORTFM_FORCEINLINE inline
+#define UE_AUTORTFM_FORCEINLINE_ALWAYS inline
+#define UE_AUTORTFM_FORCENOINLINE
+#define UE_AUTORTFM_ASSUME(x)
+#endif
+
+
+#if (defined(UE_BUILD_DEBUG) && UE_BUILD_DEBUG) || (defined(AUTORTFM_BUILD_DEBUG) && AUTORTFM_BUILD_DEBUG)
+// Force-inlining can make debugging glitchy. Disable this if we're running a debug build.
+#undef UE_AUTORTFM_CALLSITE_FORCEINLINE
+#define UE_AUTORTFM_CALLSITE_FORCEINLINE
+#undef UE_AUTORTFM_FORCEINLINE
+#define UE_AUTORTFM_FORCEINLINE inline
+#endif
+
+#ifdef _MSC_VER
+#define AUTORTFM_DISABLE_UNREACHABLE_CODE_WARNINGS \
+	__pragma (warning(push)) \
+	__pragma (warning(disable: 4702)) /* unreachable code */
+#define AUTORTFM_RESTORE_UNREACHABLE_CODE_WARNINGS \
+	__pragma (warning(pop))
+#else
+#define AUTORTFM_DISABLE_UNREACHABLE_CODE_WARNINGS
+#define AUTORTFM_RESTORE_UNREACHABLE_CODE_WARNINGS
+#endif
+
+#define UE_AUTORTFM_UNUSED(UNUSEDVAR) (void)UNUSEDVAR
+
+// It is critical that these functions are both static and forceinline to prevent binary sizes to explode
+// This is a trick to ensure that there will never be a non-inlined version of these functions that the linker can decide to use
+#ifndef UE_HEADER_UNITS
+#define UE_AUTORTFM_CRITICAL_INLINE static UE_AUTORTFM_FORCEINLINE
+#define UE_AUTORTFM_CRITICAL_INLINE_ALWAYS static UE_AUTORTFM_FORCEINLINE_ALWAYS
+#else
+#define UE_AUTORTFM_CRITICAL_INLINE UE_AUTORTFM_FORCEINLINE // TODO: This needs to be revisited. we don't want bloated executables when modules are enabled
+#define UE_AUTORTFM_CRITICAL_INLINE_ALWAYS UE_AUTORTFM_FORCEINLINE_ALWAYS // TODO: This needs to be revisited. we don't want bloated executables when modules are enabled
+#endif
diff --git a/Engine/Source/Runtime/AutoRTFM/Public/AutoRTFM/Metrics.h b/Engine/Source/Runtime/AutoRTFM/Public/AutoRTFM/Metrics.h
new file mode 100644
--- /dev/null
+++ b/Engine/Source/Runtime/AutoRTFM/Public/AutoRTFM/Metrics.h
@@ -0,0 +1,26 @@
+// Copyright Epic Games, Inc. All Rights Reserved.
+
+#pragma once
+
+#include <cstdint>
+
+namespace AutoRTFM
+{
+
+struct FAutoRTFMMetrics
+{
+	int64_t NumTransactionsStarted;
+	int64_t NumTransactionsCommitted;
+
+	int64_t NumTransactionsAborted;
+	int64_t NumTransactionsAbortedByRequest;
+	int64_t NumTransactionsAbortedByLanguage;
+};
+
+// reset the internal metrics tracking back to zero
+void ResetAutoRTFMMetrics();
+
+// get a snapshot of the current internal metrics
+FAutoRTFMMetrics GetAutoRTFMMetrics();
+
+}
diff --git a/Engine/Source/Runtime/AutoRTFM/Public/AutoRTFM/Task.h b/Engine/Source/Runtime/AutoRTFM/Public/AutoRTFM/Task.h
new file mode 100644
--- /dev/null
+++ b/Engine/Source/Runtime/AutoRTFM/Public/AutoRTFM/Task.h
@@ -0,0 +1,252 @@
+// Copyright Epic Games, Inc. All Rights Reserved.
+
+#pragma once
+
+#include <cstddef>
+#include <type_traits>
+#include <utility>
+
+#include "AutoRTFM/Defines.h"
+
+namespace AutoRTFM
+{
+
+// TTask is a small std::function like class that can wrap a functor (lambda, etc).
+// TTask is used by AutoRTFM for callbacks such as OnCommit and OnAbort
+// handlers, which there may be many accumulating throughout a transaction.
+// For this reason TTask is designed to be compact in size (smaller than 
+// TFunction), but large enough to fit the most common functors used by AutoRTFM
+// tasks.
+// Note: to prevent generating closed-variants of task-wrapped functions,
+// operator() cannot be called from the closed. All other operations on
+// the task can be called from the open or closed.
+template<typename Signature>
+class TTask;
+
+template<typename ReturnType, typename... ParameterTypes>
+class TTask<ReturnType(ParameterTypes...)>
+{
+public:
+	static constexpr size_t InlineDataSize = 16;
+	static constexpr size_t InlineDataAlignment = 8;
+
+private:
+	// The payload holds either the inline function data, or a pointer to the
+	// data in a heap allocation.
+	union FPayload
+	{
+		std::byte Inline[InlineDataSize];
+		void* External;
+	};
+
+	// Holds the function pointers for operating on the payload.
+	struct FPayloadMethods
+	{
+		using CallFn = ReturnType(FPayload& Payload, ParameterTypes...);
+		using CopyFn = void(FPayload& Dst, const FPayload& Src);
+		using MoveFn = void(FPayload& Dst, FPayload& Src);
+		using DestructFn = void(FPayload& Payload);
+
+		CallFn* Call = nullptr; // Calls the function held by the payload.
+		CopyFn* Copy = nullptr; // Copies the payload.
+		MoveFn* Move = nullptr; // Moves the payload.
+		DestructFn* Destruct = nullptr; // Destructs the payload.
+	};
+	
+	// Traits helper for a functor (T).
+	template<typename T>
+	struct TTraits
+	{
+		// True iff T can fit within the inline payload, and does not require
+		// courser alignment than provided by the inline payload.
+		static constexpr bool bCanUseInline = sizeof(T) <= InlineDataSize && alignof(T) <= InlineDataAlignment;
+
+		// Returns the functor from the payload.
+		static T* FunctionFrom(FPayload& Payload)
+		{
+			return bCanUseInline ? reinterpret_cast<T*>(Payload.Inline) : reinterpret_cast<T*>(Payload.External);
+		}
+
+		// Returns the functor from the payload.
+		static const T* FunctionFrom(const FPayload& Payload)
+		{
+			return bCanUseInline ? reinterpret_cast<const T*>(Payload.Inline) : reinterpret_cast<const T*>(Payload.External);
+		}
+
+		// Calls the function held by the payload. Cannot be called from the closed.
+		AUTORTFM_DISABLE static ReturnType Call(FPayload& Payload, ParameterTypes... Arguments)
+		{
+			return (*FunctionFrom(Payload))(std::forward<ParameterTypes>(Arguments)...);
+		}
+
+		// Copies the payload.
+		static void Copy(FPayload& Dst, const FPayload& Src)
+		{
+			if constexpr (bCanUseInline)
+			{
+				new(Dst.Inline) T(*FunctionFrom(Src));
+			}
+			else
+			{
+				Dst.External = new T(*FunctionFrom(Src));
+			}
+		}
+
+		// Moves the payload.
+		static void Move(FPayload& Dst, FPayload& Src)
+		{
+			if constexpr (bCanUseInline)
+			{
+				new(Dst.Inline) T(std::move(*FunctionFrom(Src)));
+			}
+			else
+			{
+				Dst.External = Src.External;
+			}
+		}
+
+		// Destructs the payload.
+		static void Destruct(FPayload& Payload)
+		{
+			if constexpr (bCanUseInline)
+			{
+				FunctionFrom(Payload)->~T();
+			}
+			else
+			{
+				delete FunctionFrom(Payload);
+			}
+		}
+
+		// The methods to operate on a payload for a functor of type T.
+		static constexpr FPayloadMethods PayloadMethods { &Call, &Copy, &Move, &Destruct };
+	};
+public:
+	// Constructor.
+	// The TTask is constructed in a non-set state.
+	TTask() = default;
+
+	// Destructor.
+	~TTask()
+	{
+		Reset();
+	}
+
+	// Copy constructor
+	TTask(const TTask& Other)
+	{
+		if (Other.IsSet())
+		{
+			PayloadMethods = Other.PayloadMethods;
+			PayloadMethods->Copy(Payload, Other.Payload);
+		}
+	}
+
+	// Move constructor
+	TTask(TTask&& Other)
+	{
+		if (Other.IsSet())
+		{
+			PayloadMethods = Other.PayloadMethods;
+			PayloadMethods->Move(Payload, Other.Payload);
+			Other.PayloadMethods = nullptr;
+		}
+	}
+
+	// Constructor from functor.
+	template
+	<
+		typename FunctorType,
+		typename = std::enable_if_t
+		<
+			// Use the explicit TTask copy / move constructors instead of this constructor
+			!std::is_same_v<std::decay_t<FunctorType>, TTask> &&
+			// The templated argument must be a functor with the correct signature
+			std::is_invocable_r_v<ReturnType, std::decay_t<FunctorType>, ParameterTypes...>
+		>
+	>
+	TTask(FunctorType&& Other)
+	{
+		using Decayed = std::decay_t<FunctorType>;
+		using Traits = TTraits<Decayed>;
+		PayloadMethods = &Traits::PayloadMethods;
+		if (Traits::bCanUseInline)
+		{
+			new(Payload.Inline) Decayed(std::forward<FunctorType>(Other));
+		}
+		else
+		{
+			Payload.External = new Decayed(std::forward<FunctorType>(Other));
+		}
+	}
+
+	// Copy assignment operator
+	TTask& operator = (const TTask& Other)
+	{
+		if (&Other != this)
+		{
+			Reset();
+			if (Other.IsSet())
+			{
+				PayloadMethods = Other.PayloadMethods;
+				PayloadMethods->Copy(Payload, Other.Payload);
+			}
+		}
+		return *this;
+	}
+
+	// Move assignment operator
+	TTask& operator = (TTask&& Other)
+	{
+		if (&Other != this)
+		{
+			Reset();
+			if (Other.IsSet())
+			{
+				PayloadMethods = Other.PayloadMethods;
+				PayloadMethods->Move(Payload, Other.Payload);
+				Other.PayloadMethods = nullptr;
+			}
+		}
+		return *this;
+	}
+
+	// Call operator. Cannot be called from the closed.
+	template<typename... ArgumentTypes>
+	AUTORTFM_DISABLE ReturnType operator()(ArgumentTypes&&... Arguments) const
+	{
+		return PayloadMethods->Call(Payload, std::forward<ArgumentTypes>(Arguments)...);
+	}
+
+	// Reset the task to an unset state.
+	void Reset()
+	{
+		if (PayloadMethods)
+		{
+			PayloadMethods->Destruct(Payload);
+			PayloadMethods = nullptr;
+		}
+	}
+
+	// Returns true if the task holds a functor.
+	bool IsSet() const
+	{
+		return PayloadMethods != nullptr;
+	}
+
+	// Returns the address of the functor.
+	void* FunctionAddress() const
+	{
+		return reinterpret_cast<void*>(PayloadMethods->Call);
+	}
+
+private:
+	// mutable as the functor might be mutable.
+	alignas(InlineDataAlignment) mutable FPayload Payload;
+	FPayloadMethods const* PayloadMethods = nullptr;
+};
+
+// Tasks are designed to be compact.
+static_assert(sizeof(TTask<void()>) <= 24);
+
+} // namespace AutoRTFM
diff --git a/Engine/Source/Runtime/AutoRTFM/Public/AutoRTFM/Testing.h b/Engine/Source/Runtime/AutoRTFM/Public/AutoRTFM/Testing.h
new file mode 100644
--- /dev/null
+++ b/Engine/Source/Runtime/AutoRTFM/Public/AutoRTFM/Testing.h
@@ -0,0 +1,49 @@
+// Copyright Epic Games, Inc. All Rights Reserved.
+
+#pragma once
+
+#include "AutoRTFM.h"
+
+namespace AutoRTFM::Testing
+{
+
+// Implemented by the test framework, and called when AUTORTFM_TESTING_ASSERT() fails.
+void AssertionFailure(const char* Expression, const char* File, int Line);
+
+#define AUTORTFM_TESTING_ASSERT(Condition) \
+	do { if (!(Condition)) { AssertionFailure("AUTORTFM_TESTING_ASSERT(" #Condition ")", __FILE__, __LINE__); }} while (false)
+
+// Run the callback in a transaction like Transact, but abort program
+// execution if the result is anything other than autortfm_committed.
+// Useful for testing.
+template<typename TFunctor>
+static UE_AUTORTFM_FORCEINLINE void Commit(AUTORTFM_IMPLICIT_ENABLE const TFunctor& Functor)
+{
+	const ETransactionResult Result = Transact(Functor);
+	AUTORTFM_TESTING_ASSERT(ETransactionResult::Committed == Result);
+}
+
+// Run the callback in a transaction like Transact, but abort program
+// execution if the result is anything other than abort.
+template<typename TFunctor>
+static UE_AUTORTFM_FORCEINLINE void Abort(AUTORTFM_IMPLICIT_ENABLE const TFunctor& Functor)
+{
+	const ETransactionResult Result = Transact(Functor);
+	AUTORTFM_TESTING_ASSERT(ETransactionResult::Committed != Result);
+}
+
+// Force set the AutoRTFM runtime state. Used in testing. Returns the old value.
+UE_AUTORTFM_API ForTheRuntime::EAutoRTFMEnabledState ForceSetAutoRTFMRuntime(ForTheRuntime::EAutoRTFMEnabledState State);
+
+struct FEnabledStateResetterScoped final
+{
+	FEnabledStateResetterScoped(ForTheRuntime::EAutoRTFMEnabledState State) : Original(ForceSetAutoRTFMRuntime(State)) {}
+	~FEnabledStateResetterScoped() { ForceSetAutoRTFMRuntime(Original); }
+
+private:
+	const ForTheRuntime::EAutoRTFMEnabledState Original;
+};
+
+#undef AUTORTFM_TESTING_ASSERT
+
+} // namespace AutoRTFM::Testing
diff --git a/Engine/Source/Runtime/AutoRTFM/Public/AutoRTFMConstants.h b/Engine/Source/Runtime/AutoRTFM/Public/AutoRTFMConstants.h
deleted file mode 100644
diff --git a/Engine/Source/Runtime/AutoRTFM/Public/AutoRTFMDefines.h b/Engine/Source/Runtime/AutoRTFM/Public/AutoRTFMDefines.h
deleted file mode 100644
diff --git a/Engine/Source/Runtime/AutoRTFM/Public/AutoRTFMMetrics.h b/Engine/Source/Runtime/AutoRTFM/Public/AutoRTFMMetrics.h
deleted file mode 100644
diff --git a/Engine/Source/Runtime/AutoRTFM/Public/AutoRTFMTask.h b/Engine/Source/Runtime/AutoRTFM/Public/AutoRTFMTask.h
deleted file mode 100644
diff --git a/Engine/Source/Runtime/AutoRTFM/Public/AutoRTFMTesting.h b/Engine/Source/Runtime/AutoRTFM/Public/AutoRTFMTesting.h
deleted file mode 100644
diff --git a/Engine/Source/Runtime/Engine/Private/Tests/AutoRTFM/AutoRTFMActorComponentTests.cpp b/Engine/Source/Runtime/Engine/Private/Tests/AutoRTFM/AutoRTFMActorComponentTests.cpp
--- a/Engine/Source/Runtime/Engine/Private/Tests/AutoRTFM/AutoRTFMActorComponentTests.cpp
+++ b/Engine/Source/Runtime/Engine/Private/Tests/AutoRTFM/AutoRTFMActorComponentTests.cpp
@@ -3,7 +3,7 @@
 #include "Misc/AutomationTest.h"
 #include "Animation/AnimInstance.h"
 #include "Animation/Skeleton.h"
-#include "AutoRTFMTesting.h"
+#include "AutoRTFM/Testing.h"
 #include "AutoRTFMTestActor.h"
 #include "AutoRTFMTestAnotherActor.h"
 #include "AutoRTFMTestBodySetup.h"
diff --git a/Engine/Source/Runtime/Engine/Private/Tests/AutoRTFM/AutoRTFMNetworkProfilerTests.cpp b/Engine/Source/Runtime/Engine/Private/Tests/AutoRTFM/AutoRTFMNetworkProfilerTests.cpp
--- a/Engine/Source/Runtime/Engine/Private/Tests/AutoRTFM/AutoRTFMNetworkProfilerTests.cpp
+++ b/Engine/Source/Runtime/Engine/Private/Tests/AutoRTFM/AutoRTFMNetworkProfilerTests.cpp
@@ -1,6 +1,6 @@
 // Copyright Epic Games, Inc. All Rights Reserved.
 
-#include "AutoRTFMTesting.h"
+#include "AutoRTFM/Testing.h"
 #include "AutoRTFMTestActor.h"
 #include "Engine/ActorChannel.h"
 #include "Engine/DemoNetDriver.h"
diff --git a/Engine/Source/Runtime/Engine/Private/Tests/AutoRTFM/AutoRTFMSlateTests.cpp b/Engine/Source/Runtime/Engine/Private/Tests/AutoRTFM/AutoRTFMSlateTests.cpp
--- a/Engine/Source/Runtime/Engine/Private/Tests/AutoRTFM/AutoRTFMSlateTests.cpp
+++ b/Engine/Source/Runtime/Engine/Private/Tests/AutoRTFM/AutoRTFMSlateTests.cpp
@@ -1,6 +1,6 @@
 // Copyright Epic Games, Inc. All Rights Reserved.
 
-#include "AutoRTFMTesting.h"
+#include "AutoRTFM/Testing.h"
 #include "Tests/TestHarnessAdapter.h"
 #include "AutoRTFM.h"
 
diff --git a/Engine/Source/Runtime/Engine/Private/Tests/AutoRTFM/AutoRTFMTestingAssertionFailure.cpp b/Engine/Source/Runtime/Engine/Private/Tests/AutoRTFM/AutoRTFMTestingAssertionFailure.cpp
--- a/Engine/Source/Runtime/Engine/Private/Tests/AutoRTFM/AutoRTFMTestingAssertionFailure.cpp
+++ b/Engine/Source/Runtime/Engine/Private/Tests/AutoRTFM/AutoRTFMTestingAssertionFailure.cpp
@@ -1,6 +1,6 @@
 // Copyright Epic Games, Inc. All Rights Reserved.
 
-#include "AutoRTFMTesting.h"
+#include "AutoRTFM/Testing.h"
 #include "Misc/AssertionMacros.h"
 
 #if WITH_DEV_AUTOMATION_TESTS
