[AutoRTFM] Move UE_AUTORTFM_API define from Defines.h to CAPI.h

While it is a define, it's only used by the AutoRTFM APIs.
Moved because there are modules that do not depend on Core and need to include `AutoRTFM/Defines.h`, but currently cannot due to the include `HAL/Platform.h`.

diff --git a/Engine/Source/Runtime/AutoRTFM/Public/AutoRTFM.h b/Engine/Source/Runtime/AutoRTFM/Public/AutoRTFM.h
--- a/Engine/Source/Runtime/AutoRTFM/Public/AutoRTFM.h
+++ b/Engine/Source/Runtime/AutoRTFM/Public/AutoRTFM.h
@@ -11,7 +11,7 @@
 // HEADER_UNIT_SKIP - unused warnings
 
 #include "AutoRTFM/Constants.h"
-#include "AutoRTFM/Defines.h" // IWYU pragma: export
+#include "AutoRTFM/Defines.h"
 #include "AutoRTFM/CAPI.h"
 
 #ifdef __cplusplus
diff --git a/Engine/Source/Runtime/AutoRTFM/Public/AutoRTFM/CAPI.h b/Engine/Source/Runtime/AutoRTFM/Public/AutoRTFM/CAPI.h
--- a/Engine/Source/Runtime/AutoRTFM/Public/AutoRTFM/CAPI.h
+++ b/Engine/Source/Runtime/AutoRTFM/Public/AutoRTFM/CAPI.h
@@ -7,6 +7,18 @@
 #include <cstdarg> // for va_list
 #include <cstddef> // for size_t
 
+// Define UE_AUTORTFM_API (DLL export / import attributes)
+#if defined(__cplusplus) && defined(__UNREAL__) && !(defined(UE_AUTORTFM_DO_NOT_INCLUDE_PLATFORM_H) && UE_AUTORTFM_DO_NOT_INCLUDE_PLATFORM_H)
+	// Include HAL/Platform.h for DLLIMPORT / DLLEXPORT definitions, which
+	// UBT can use as a definition for AUTORTFM_API.
+	#include <HAL/Platform.h>
+	#define UE_AUTORTFM_API AUTORTFM_API
+#else
+	#ifndef UE_AUTORTFM_API
+		#define UE_AUTORTFM_API
+	#endif
+#endif
+
 #ifdef __cplusplus
 extern "C"
 {
diff --git a/Engine/Source/Runtime/AutoRTFM/Public/AutoRTFM/Defines.h b/Engine/Source/Runtime/AutoRTFM/Public/AutoRTFM/Defines.h
--- a/Engine/Source/Runtime/AutoRTFM/Public/AutoRTFM/Defines.h
+++ b/Engine/Source/Runtime/AutoRTFM/Public/AutoRTFM/Defines.h
@@ -1,8 +1,5 @@
 // Copyright Epic Games, Inc. All Rights Reserved.
 
-// IWYU pragma: private
-// Include include "AutoRTFM.h" instead
-
 #pragma once
 
 #include "AutoRTFM/Constants.h"
@@ -154,17 +151,6 @@
 #define UE_AUTORTFM_CONCAT_IMPL(A, B) A ## B
 #define UE_AUTORTFM_CONCAT(A, B) UE_AUTORTFM_CONCAT_IMPL(A, B)
 
-#if defined(__cplusplus) && defined(__UNREAL__) && !(defined(UE_AUTORTFM_DO_NOT_INCLUDE_PLATFORM_H) && UE_AUTORTFM_DO_NOT_INCLUDE_PLATFORM_H)
-	// Include HAL/Platform.h for DLLIMPORT / DLLEXPORT definitions, which
-	// UBT can use as a definition for AUTORTFM_API.
-	#include <HAL/Platform.h>
-	#define UE_AUTORTFM_API AUTORTFM_API
-#else
-	#ifndef UE_AUTORTFM_API
-		#define UE_AUTORTFM_API
-	#endif
-#endif
-
 #if defined(_MSC_VER)
 #define UE_AUTORTFM_FORCEINLINE __forceinline
 #define UE_AUTORTFM_FORCEINLINE_ALWAYS __forceinline
