Remove #include <memory.h>

This is _not_ a standard library header! <memory> is the closest equivalent. It doesn't seem that we actually need this to compile, regardless.

diff --git a/Engine/Source/Runtime/AutoRTFM/Public/AutoRTFM.h b/Engine/Source/Runtime/AutoRTFM/Public/AutoRTFM.h
--- a/Engine/Source/Runtime/AutoRTFM/Public/AutoRTFM.h
+++ b/Engine/Source/Runtime/AutoRTFM/Public/AutoRTFM.h
@@ -22,7 +22,6 @@
 #include <new>
 #include <type_traits>
 #include <utility>
-#include <memory.h>
 #include <stdlib.h>
 
 namespace AutoRTFM
