[AutoRTFM] static_assert that TTask return type matches functor return type

Prevents attempting to return values in lambdas that would just get dropped on the floor.

diff --git a/Engine/Source/Runtime/AutoRTFM/Public/AutoRTFM/Task.h b/Engine/Source/Runtime/AutoRTFM/Public/AutoRTFM/Task.h
--- a/Engine/Source/Runtime/AutoRTFM/Public/AutoRTFM/Task.h
+++ b/Engine/Source/Runtime/AutoRTFM/Public/AutoRTFM/Task.h
@@ -165,18 +165,20 @@ public:
 			std::is_invocable_r_v<ReturnType, std::decay_t<FunctorType>, ParameterTypes...>
 		>
 	>
-	TTask(FunctorType&& Other)
+	TTask(FunctorType&& Functor)
 	{
 		using Decayed = std::decay_t<FunctorType>;
 		using Traits = TTraits<Decayed>;
+		static_assert(std::is_same_v<std::invoke_result_t<Decayed, ParameterTypes...>, ReturnType>,
+			"Functor return type does not match TTask return type");
 		PayloadMethods = &Traits::PayloadMethods;
 		if (Traits::bCanUseInline)
 		{
-			new(Payload.Inline) Decayed(std::forward<FunctorType>(Other));
+			new(Payload.Inline) Decayed(std::forward<FunctorType>(Functor));
 		}
 		else
 		{
-			Payload.External = new Decayed(std::forward<FunctorType>(Other));
+			Payload.External = new Decayed(std::forward<FunctorType>(Functor));
 		}
 	}
 
