[AutoRTFM] Improve AutoRTFM error reporting when accessing a transactionally-unsafe FMutex from the closed.

Previously, the error was inscrutable and only made sense if you were aware that FMutex is based on atomics.

diff --git a/Engine/Source/Programs/AutoRTFMTests/Private/Abort.cpp b/Engine/Source/Programs/AutoRTFMTests/Private/Abort.cpp
--- a/Engine/Source/Programs/AutoRTFMTests/Private/Abort.cpp
+++ b/Engine/Source/Programs/AutoRTFMTests/Private/Abort.cpp
@@ -4,6 +4,7 @@
 #include "AutoRTFM/Testing.h"
 #include "AutoRTFMTestUtils.h"
 #include "AutoRTFM/AutoRTFMUE.h"
+#include "Async/Mutex.h"
 #include "Catch2Includes.h"
 #include "Misc/StringOutputDevice.h"
 #include "Misc/OutputDeviceRedirector.h"
@@ -1443,6 +1444,42 @@ TEST_CASE("Abort.Language.ErrorHints.TryAcquireSRWLockShared")
 	REQUIRE(WarningContext.HasWarningSubstring("A read-write lock was accessed while inside an AutoRTFM transaction."));
 }
 
+TEST_CASE("Abort.Language.ErrorHints.MutexLock")
+{
+	FScopedErrorHintTest ScopedTest;
+    AutoRTFMTestUtils::FCaptureWarningContext WarningContext;
+
+	UE::FMutex Mutex;
+
+	const AutoRTFM::ETransactionResult Result = AutoRTFM::Transact([&]
+	{
+		Mutex.Lock();
+	});
+
+	REQUIRE(AutoRTFM::ETransactionResult::AbortedByLanguage == Result);
+
+	// Did we print our message?
+	REQUIRE(WarningContext.HasWarningSubstring("A mutex was accessed while inside an AutoRTFM transaction."));
+}
+
+TEST_CASE("Abort.Language.ErrorHints.MutexIsLocked")
+{
+	FScopedErrorHintTest ScopedTest;
+    AutoRTFMTestUtils::FCaptureWarningContext WarningContext;
+
+	UE::FMutex Mutex;
+
+	const AutoRTFM::ETransactionResult Result = AutoRTFM::Transact([&]
+	{
+		REQUIRE(!Mutex.IsLocked());
+	});
+
+	REQUIRE(AutoRTFM::ETransactionResult::AbortedByLanguage == Result);
+
+	// Did we print our message?
+	REQUIRE(WarningContext.HasWarningSubstring("A mutex was accessed while inside an AutoRTFM transaction."));
+}
+
 #if 0
 // Testing an open free is flaky, since it leads to a double free--one inside the open block, and one during the
 // transactional abort caused by an AutoRTFM hazard. This can cause the unit test harness to fail on a subsequent test,
diff --git a/Engine/Source/Runtime/AutoRTFM/Private/FunctionMap.cpp b/Engine/Source/Runtime/AutoRTFM/Private/FunctionMap.cpp
--- a/Engine/Source/Runtime/AutoRTFM/Private/FunctionMap.cpp
+++ b/Engine/Source/Runtime/AutoRTFM/Private/FunctionMap.cpp
@@ -5,9 +5,9 @@
 #include "FunctionMap.h"
 
 #include "AutoRTFM.h"
+#include "AutoRTFM/ReportHazard.h"
 #include "BuildMacros.h"
 #include "FunctionMapInlines.h"
-#include "ReportAutoRTFMHazard.h"
 #include "Utils.h"
 
 #include <cinttypes>
@@ -271,6 +271,7 @@ AUTORTFM_DISABLE static void* FunctionMapReportError(void* Open, const char* Whe
 		Open == &DeleteCriticalSection)
 	{
 		ForTheRuntime::ReportAutoRTFMHazard(ForTheRuntime::EHazardType::CriticalSection);
+		return nullptr;
 	}
 	else if (Open == &TryAcquireSRWLockExclusive ||
 		Open == &AcquireSRWLockExclusive ||
@@ -281,6 +282,7 @@ AUTORTFM_DISABLE static void* FunctionMapReportError(void* Open, const char* Whe
 		Open == &TryAcquireSRWLockExclusive)
 	{
 		ForTheRuntime::ReportAutoRTFMHazard(ForTheRuntime::EHazardType::RWLock);
+		return nullptr;
 	}
 #elif AUTORTFM_PLATFORM_LINUX
 	if (Open == &pthread_mutexattr_init ||
@@ -293,6 +295,7 @@ AUTORTFM_DISABLE static void* FunctionMapReportError(void* Open, const char* Whe
 		Open == &pthread_mutex_unlock)
 	{
 		ForTheRuntime::ReportAutoRTFMHazard(ForTheRuntime::EHazardType::CriticalSection);
+		return nullptr;
 	}
 	else if (Open == &pthread_rwlock_init ||
 		Open == &pthread_rwlock_destroy ||
@@ -303,6 +306,7 @@ AUTORTFM_DISABLE static void* FunctionMapReportError(void* Open, const char* Whe
 		Open == &pthread_rwlock_rdlock)
 	{
 		ForTheRuntime::ReportAutoRTFMHazard(ForTheRuntime::EHazardType::RWLock);
+		return nullptr;
 	}
 #endif
 
diff --git a/Engine/Source/Runtime/AutoRTFM/Private/ReportAutoRTFMHazard.cpp b/Engine/Source/Runtime/AutoRTFM/Private/ReportAutoRTFMHazard.cpp
deleted file mode 100644
diff --git a/Engine/Source/Runtime/AutoRTFM/Private/ReportAutoRTFMHazard.h b/Engine/Source/Runtime/AutoRTFM/Private/ReportAutoRTFMHazard.h
deleted file mode 100644
diff --git a/Engine/Source/Runtime/AutoRTFM/Private/ReportHazard.cpp b/Engine/Source/Runtime/AutoRTFM/Private/ReportHazard.cpp
new file mode 100644
--- /dev/null
+++ b/Engine/Source/Runtime/AutoRTFM/Private/ReportHazard.cpp
@@ -0,0 +1,76 @@
+// Copyright Epic Games, Inc. All Rights Reserved.
+
+#if (defined(__AUTORTFM) && __AUTORTFM)
+
+#include "AutoRTFM/ReportHazard.h"
+#include "Utils.h"
+
+void AutoRTFM::ForTheRuntime::ReportAutoRTFMHazard(EHazardType HazardType)
+{
+	static constexpr const char Banner[] = R"(
+             _         _        ____ _____ _____ __  __   _   _                        _
+            / \  _   _| |_ ___ |  _ \_   _|  ___|  \/  | | | | | __ _ ______ _ _ __ __| |
+  _____    / _ \| | | | __/ _ \| |_) || | | |_  | |\/| | | |_| |/ _` |_  / _` | '__/ _` |  _____
+ |_____|  / ___ \ |_| | || (_) |  _ < | | |  _| | |  | | |  _  | (_| |/ / (_| | | | (_| | |_____|
+         /_/   \_\__,_|\__\___/|_| \_\|_| |_|   |_|  |_| |_| |_|\__,_/___\__,_|_|  \__,_|
+
+New to AutoRTFM? See Engine/Source/Runtime/AutoRTFM/Documentation/README.md for more info.
+)";
+
+	const char* Message = "An unknown error occurred.";
+
+	switch (HazardType)
+	{
+		case EHazardType::CriticalSection:
+			Message = R"(
+A critical section was accessed while inside an AutoRTFM transaction. AutoRTFM doesn't support native
+critical sections, but it does support transactionally-safe critical sections that can be used instead.
+
+Use the call stack below to find the FCriticalSection accessed immediately before this error occurred, and
+replace it with FTransactionallySafeCriticalSection. If FScopeLock is used, replace it with UE::TScopeLock.
+)";
+			break;
+
+		case EHazardType::RWLock:
+			Message = R"(
+A read-write lock was accessed while inside an AutoRTFM transaction. AutoRTFM doesn't support native
+read-write locks, but it does support transactionally-safe read-write locks that can be used instead.
+
+Use the call stack below to find the FRWLock accessed immediately before this error occurred, and
+replace it with FTransactionallySafeFRWLock. If FReadScopeLock, FWriteScopeLock or FRWScopeLock are
+used, replace them with UE::TReadScopeLock, UE::TWriteScopeLock or UE::TRWScopeLock.
+)";
+			break;
+
+		case EHazardType::Mutex:
+			Message = R"(
+A mutex was accessed while inside an AutoRTFM transaction. AutoRTFM doesn't support native mutexes, but
+it does support transactionally-safe mutexes that can be used instead.
+
+Use the call stack below to find the FMutex accessed immediately before this error occurred, and
+replace it with FTransactionallySafeMutex.
+)";
+			break;
+
+		case EHazardType::OpenFree:
+			Message = R"(
+Memory has been freed from the open while inside an AutoRTFM transaction. This memory was allocated
+or written-to from the closed within the same transaction; freeing it makes rollback impossible.
+
+Use the call stack below to find the location of the free call; the open block will be further up in
+the stack. If possible, avoid going into the open at all; if it is unavoidable, ensure that memory
+is never modified or freed from the open if it might also be written-to from the closed.
+)";
+			break;
+
+		default:
+			AUTORTFM_ENSURE_MSG(false, "Unrecognized hazard type %d", (int)HazardType);
+			break;
+	}
+
+	AUTORTFM_WARN("%s%s", Banner, Message);
+	::AutoRTFM::LogWithCallstack(autortfm_log_warn, "Call stack:");
+	AUTORTFM_REPORT_ERROR("Failure was caused by an AutoRTFM hazard.");
+}
+
+#endif // defined(__AUTORTFM) && __AUTORTFM
diff --git a/Engine/Source/Runtime/AutoRTFM/Private/TransactionInlines.h b/Engine/Source/Runtime/AutoRTFM/Private/TransactionInlines.h
--- a/Engine/Source/Runtime/AutoRTFM/Private/TransactionInlines.h
+++ b/Engine/Source/Runtime/AutoRTFM/Private/TransactionInlines.h
@@ -10,7 +10,7 @@
 #include "Transaction.h"
 #include "Utils.h"
 #include "WriteLog.h"
-#include "ReportAutoRTFMHazard.h"
+#include "AutoRTFM/ReportHazard.h"
 
 #include <utility>
 
diff --git a/Engine/Source/Runtime/AutoRTFM/Public/AutoRTFM/ReportHazard.h b/Engine/Source/Runtime/AutoRTFM/Public/AutoRTFM/ReportHazard.h
new file mode 100644
--- /dev/null
+++ b/Engine/Source/Runtime/AutoRTFM/Public/AutoRTFM/ReportHazard.h
@@ -0,0 +1,43 @@
+// Copyright Epic Games, Inc. All Rights Reserved.
+
+#pragma once
+
+#include "AutoRTFM/CAPI.h"
+#include "AutoRTFM/Defines.h"
+
+#if UE_AUTORTFM
+
+namespace AutoRTFM::ForTheRuntime
+{
+
+enum class EHazardType
+{
+	CriticalSection,
+	RWLock,
+	Mutex,
+	OpenFree,
+};
+
+/**
+ * Prints an AutoRTFM Hazard banner, with descriptive text based on the hazard type, and a callstack.
+ * Concludes with AutoRTFM::ReportError, which will terminate the process under normal circumstances.
+ */
+AUTORTFM_OPEN void ReportAutoRTFMHazard(EHazardType HazardType);
+
+}
+
+/**
+ * Reports an AutoRTFM hazard if this block of code is reached from the closed.
+ * `autortfm_is_closed` is optimized away at compile time, so this compiles away to nothing in an open function.
+ */
+#define UE_AUTORTFM_REPORT_HAZARD_IF_CLOSED(InHazardType)                                                           \
+	(autortfm_is_closed()                                                                                       \
+		? ::AutoRTFM::ForTheRuntime::ReportAutoRTFMHazard(::AutoRTFM::ForTheRuntime::EHazardType::InHazardType) \
+		: (void)0)
+
+#else
+
+/** Our macro compiles away to nothing on compilers which do not support AutoRTFM. */
+#define UE_AUTORTFM_REPORT_HAZARD_IF_CLOSED(InHazardType) ((void)0)
+
+#endif // UE_AUTORTFM
\ No newline at end of file
diff --git a/Engine/Source/Runtime/Core/Public/Async/Mutex.h b/Engine/Source/Runtime/Core/Public/Async/Mutex.h
--- a/Engine/Source/Runtime/Core/Public/Async/Mutex.h
+++ b/Engine/Source/Runtime/Core/Public/Async/Mutex.h
@@ -3,6 +3,7 @@
 #pragma once
 
 #include "Async/LockTags.h"
+#include "AutoRTFM/ReportHazard.h"
 #include "CoreTypes.h"
 #include <atomic>
 
@@ -30,11 +31,13 @@ public:
 
 	[[nodiscard]] inline bool IsLocked() const
 	{
+		UE_AUTORTFM_REPORT_HAZARD_IF_CLOSED(Mutex);
 		return !!(State.load(std::memory_order_relaxed) & IsLockedFlag);
 	}
 
 	[[nodiscard]] inline bool TryLock()
 	{
+		UE_AUTORTFM_REPORT_HAZARD_IF_CLOSED(Mutex);
 		uint8 Expected = State.load(std::memory_order_relaxed);
 		return !(Expected & IsLockedFlag) &&
 			State.compare_exchange_strong(Expected, Expected | IsLockedFlag, std::memory_order_acquire, std::memory_order_relaxed);
@@ -42,6 +45,7 @@ public:
 
 	inline void Lock()
 	{
+		UE_AUTORTFM_REPORT_HAZARD_IF_CLOSED(Mutex);
 		uint8 Expected = 0;
 		if (LIKELY(State.compare_exchange_weak(Expected, IsLockedFlag, std::memory_order_acquire, std::memory_order_relaxed)))
 		{
@@ -53,6 +57,7 @@ public:
 	inline void Unlock()
 	{
 		// Unlock immediately to allow other threads to acquire the lock while this thread looks for a thread to wake.
+		UE_AUTORTFM_REPORT_HAZARD_IF_CLOSED(Mutex);
 		const uint8 LastState = State.fetch_sub(IsLockedFlag, std::memory_order_release);
 		if (LIKELY(!(LastState & MayHaveWaitingLockFlag)))
 		{
