Fix Horde build error.

This function's linkage needed to be specified.

diff --git a/Engine/Source/Runtime/AutoRTFM/Public/AutoRTFM/ReportHazard.h b/Engine/Source/Runtime/AutoRTFM/Public/AutoRTFM/ReportHazard.h
--- a/Engine/Source/Runtime/AutoRTFM/Public/AutoRTFM/ReportHazard.h
+++ b/Engine/Source/Runtime/AutoRTFM/Public/AutoRTFM/ReportHazard.h
@@ -22,7 +22,7 @@ enum class EHazardType
  * Prints an AutoRTFM Hazard banner, with descriptive text based on the hazard type, and a callstack.
  * Concludes with AutoRTFM::ReportError, which will terminate the process under normal circumstances.
  */
-AUTORTFM_OPEN void ReportAutoRTFMHazard(EHazardType HazardType);
+AUTORTFM_OPEN UE_AUTORTFM_API void ReportAutoRTFMHazard(EHazardType HazardType);
 
 }
 
