[AutoRTFM] Refactor AutoRTFM::Open() to fix closed variant elision, and fix lit test failures

The recent refactor of AutoRTFM::Open() caused the compiler optimizations to elide closed variants to stop functioning. AutoRTFM::Open() has been reworked to fix this, and is now even more aggressive at removing unnecessary closed variants.

Update lit tests config to use C++20 to fix compiler tests. AutoRTFM headers now use concepts, which is a C++20 only feature.

diff --git a/Engine/Source/Programs/AutoRTFMTests/Private/MTTSAccessDetector.cpp b/Engine/Source/Programs/AutoRTFMTests/Private/MTTSAccessDetector.cpp
--- a/Engine/Source/Programs/AutoRTFMTests/Private/MTTSAccessDetector.cpp
+++ b/Engine/Source/Programs/AutoRTFMTests/Private/MTTSAccessDetector.cpp
@@ -425,7 +425,7 @@ TEMPLATE_TEST_CASE("MTTransactionallySafeAccessDetector", "[MTAccessDetector]",
 
 				AutoRTFM::Open([&]
 				{
-					OtherThread = std::thread([&]
+					OtherThread = std::thread([&] AUTORTFM_ENABLE // SOL-8741
 						{
 							// Let the main thread progress.
 							Handshake += 1;
diff --git a/Engine/Source/Programs/AutoRTFMTests/Private/OpenAPI.cpp b/Engine/Source/Programs/AutoRTFMTests/Private/OpenAPI.cpp
--- a/Engine/Source/Programs/AutoRTFMTests/Private/OpenAPI.cpp
+++ b/Engine/Source/Programs/AutoRTFMTests/Private/OpenAPI.cpp
@@ -1010,7 +1010,7 @@ TEST_CASE("OpenAPI.CheckRaceAgainstOtherThread")
 	{
 		AutoRTFM::Open([&]
 		{
-			std::thread([&]
+			std::thread([&] AUTORTFM_ENABLE // SOL-8741
 			{
 				// This should be a no-op in the spawnee thread.
 				AutoRTFM::RecordOpenWrite(&bHit);
diff --git a/Engine/Source/Runtime/AutoRTFM/Public/AutoRTFM.h b/Engine/Source/Runtime/AutoRTFM/Public/AutoRTFM.h
--- a/Engine/Source/Runtime/AutoRTFM/Public/AutoRTFM.h
+++ b/Engine/Source/Runtime/AutoRTFM/Public/AutoRTFM.h
@@ -315,7 +315,7 @@ UE_AUTORTFM_CRITICAL_INLINE void CascadingAbortTransaction() {}
 // TAssignFromOpenToClosed must have a specialization for the type that is being returned.
 template
 <
-	EMemoryValidationLevel VALIDATION_LEVEL = EMemoryValidationLevel::Default,
+	EMemoryValidationLevel ValidationLevel = EMemoryValidationLevel::Default,
 	typename FunctorType = void,
 	typename ReturnType = decltype(std::declval<FunctorType>()())
 >
@@ -327,27 +327,31 @@ UE_AUTORTFM_CRITICAL_INLINE ReturnType Open(const FunctorType& Functor)
 		return Functor();
 	}
 
-	static constexpr bool bDisableValidation = (VALIDATION_LEVEL == EMemoryValidationLevel::Disabled);
-
-	auto CallOpen = [](AUTORTFM_IMPLICIT_DISABLE auto&& F)
-	{
-		if constexpr (bDisableValidation)
-		{
-			[&] AUTORTFM_OPEN_NO_VALIDATION { UE_AUTORTFM_CALLSITE_FORCEINLINE F(); }();
-		}
-		else
-		{
-			[&] AUTORTFM_OPEN { UE_AUTORTFM_CALLSITE_FORCEINLINE F(); }();
-		}
-	};
-
+// Macro filth to either call autortfm_open() or autortfm_open_explicit_validation()
+// depending on ValidationLevel. Cannot be done as a lambda or inner class method as
+// any call indirection will break compiler optimizations.
+#define AUTORTFM_INTERNAL_CALL_IN_OPEN(Arg, Function) \
+	do { \
+		if constexpr (ValidationLevel == EMemoryValidationLevel::Default) \
+		{ \
+			autortfm_open(Function, Arg, __builtin_return_address(0)); \
+		} \
+		else \
+		{ \
+			autortfm_open_explicit_validation( \
+				static_cast<autortfm_memory_validation_level>(ValidationLevel), \
+				Function, Arg, __builtin_return_address(0)); \
+		} \
+	} while (false)
+	
 	if constexpr (std::is_same_v<void, ReturnType>)
 	{
-		CallOpen([&]
+		auto Call = [] AUTORTFM_DISABLE (void* ArgPtr)
 		{
-			// Call Functor in the open
-			Functor();
-		});
+			auto& Functor = *reinterpret_cast<FunctorType*>(ArgPtr);
+			UE_AUTORTFM_CALLSITE_FORCEINLINE Functor();
+		};
+		AUTORTFM_INTERNAL_CALL_IN_OPEN(const_cast<void*>(reinterpret_cast<const void*>(&Functor)), Call);
 	}
 	else
 	{
@@ -359,12 +363,21 @@ UE_AUTORTFM_CRITICAL_INLINE ReturnType Open(const FunctorType& Functor)
 		{
 			std::aligned_storage_t<sizeof(ReturnType), alignof(ReturnType)> OpenReturnValueMemory;
 
-			CallOpen([&]
+			// Call Functor in the open, and in-place construct the return value into OpenReturnValue.
+			struct FArgs
+			{
+				const FunctorType& Functor;
+				void* OpenReturnValueMemory;
+			};
+			FArgs FunctorAndReturnValue{Functor, &OpenReturnValueMemory};
 			{
-				// Call Functor in the open, and in-place construct the return
-				// value into OpenReturnValueMemory.
-				new (&OpenReturnValueMemory) ReturnType (Functor());
-			});
+				auto Call = [] AUTORTFM_DISABLE (void* ArgPtr)
+				{
+					FArgs& Args = *reinterpret_cast<FArgs*>(ArgPtr);
+					new (Args.OpenReturnValueMemory) ReturnType (Args.Functor());
+				};
+				AUTORTFM_INTERNAL_CALL_IN_OPEN(&FunctorAndReturnValue, Call);
+			}
 
 			// Re-interpret the OpenReturnValueMemory as a ReturnType&.
 			ReturnType& OpenReturnValue = *std::launder(reinterpret_cast<ReturnType*>(&OpenReturnValueMemory));
@@ -376,10 +389,11 @@ UE_AUTORTFM_CRITICAL_INLINE ReturnType Open(const FunctorType& Functor)
 				// Destruct OpenReturnValue in the open, if required.
 				if constexpr (!std::is_trivially_destructible_v<ReturnType>)
 				{
-					CallOpen([&]
+					auto Destruct = [] AUTORTFM_DISABLE (void* ArgPtr)
 					{
-						OpenReturnValue.~ReturnType();
-					});
+						reinterpret_cast<ReturnType*>(ArgPtr)->~ReturnType();
+					};
+					AUTORTFM_INTERNAL_CALL_IN_OPEN(&OpenReturnValue, Destruct);
 				}
 				return ClosedReturnValue;
 			}
@@ -390,28 +404,37 @@ UE_AUTORTFM_CRITICAL_INLINE ReturnType Open(const FunctorType& Functor)
 				// Destruct OpenReturnValue in the open, if required.
 				if constexpr (!std::is_trivially_destructible_v<ReturnType>)
 				{
-					CallOpen([&]
+					auto Destruct = [] AUTORTFM_DISABLE (void* ArgPtr)
 					{
-						OpenReturnValue.~ReturnType();
-					});
+						reinterpret_cast<ReturnType*>(ArgPtr)->~ReturnType();
+					};
+					AUTORTFM_INTERNAL_CALL_IN_OPEN(&OpenReturnValue, Destruct);
 				}
 				return ClosedReturnValue;
 			}
 		}
 		else if constexpr (Mode == EReturnFromOpenMode::CustomMethod)
 		{
-			ReturnType ReturnValue;
-			CallOpen([&]
+			struct FArgs
+			{
+				const FunctorType& Functor;
+				ReturnType ReturnValue;
+			};
+			FArgs Args{Functor};
+			auto Call = [] AUTORTFM_DISABLE (void* ArgPtr)
 			{
-				ReturnType::AutoRTFMAssignFromOpenToClosed(ReturnValue, Functor());
-			});
-			return ReturnValue;
+				FArgs& Args = *reinterpret_cast<FArgs*>(ArgPtr);
+				ReturnType::AutoRTFMAssignFromOpenToClosed(Args.ReturnValue, Args.Functor());
+			};
+			AUTORTFM_INTERNAL_CALL_IN_OPEN(&Args, Call);
+			return Args.ReturnValue;
 		}
 		else
 		{
 			static_assert(Mode == EReturnFromOpenMode::Unsupported, "unhandled EReturnFromOpenMode");
 		}
 	}
+#undef AUTORTFM_INTERNAL_CALL_IN_OPEN
 #else // UE_AUTORTFM
 	return Functor();
 #endif // UE_AUTORTFM
