[AutoRTFM] Prevent compiler from erroring if attempting to construct a TTask in an AUTORTFM_DISABLE'd function

#Synced-CL 48337922

diff --git a/Engine/Source/Runtime/AutoRTFM/Public/AutoRTFM/Task.h b/Engine/Source/Runtime/AutoRTFM/Public/AutoRTFM/Task.h
--- a/Engine/Source/Runtime/AutoRTFM/Public/AutoRTFM/Task.h
+++ b/Engine/Source/Runtime/AutoRTFM/Public/AutoRTFM/Task.h
@@ -6,7 +6,9 @@
 #include <type_traits>
 #include <utility>
 
+#include "CAPI.h"
 #include "AutoRTFM/Defines.h"
+#include "Constants.h"
 
 namespace AutoRTFM
 {
@@ -61,6 +63,15 @@ private:
 		// courser alignment than provided by the inline payload.
 		static constexpr bool bCanUseInline = sizeof(T) <= InlineDataSize && alignof(T) <= InlineDataAlignment;
 
+		// True iff the copy constructor of the functor T is AutoRTFM-disabled.
+		static constexpr bool bCopyDisabled = AUTORTFM_MODE_OF_CALL(::new T(std::declval<const T&>())) == autortfm_mode_disable;
+
+		// True iff the move constructor of the functor T is AutoRTFM-disabled.
+		static constexpr bool bMoveDisabled = AUTORTFM_MODE_OF_CALL(::new T(std::declval<T&&>())) == autortfm_mode_disable;
+
+		// True iff the destructor of the functor T is AutoRTFM-disabled.
+		static constexpr bool bDestructorDisabled = AUTORTFM_MODE_OF_CALL(std::declval<T>().~T()) == autortfm_mode_disable;
+
 		// Returns the functor from the payload.
 		static T* FunctionFrom(FPayload& Payload)
 		{
@@ -82,6 +93,13 @@ private:
 		// Copies the payload.
 		static void Copy(FPayload& Dst, const FPayload& Src)
 		{
+			if constexpr (bCopyDisabled)
+			{
+				if (autortfm_is_closed())
+				{
+					autortfm_unreachable("TTask functor has an AutoRTFM-disabled copy constructor which was called from the closed");
+				}
+			}
 			if constexpr (bCanUseInline)
 			{
 				new(Dst.Inline) T(*FunctionFrom(Src));
@@ -95,6 +113,13 @@ private:
 		// Moves the payload.
 		static void Move(FPayload& Dst, FPayload& Src)
 		{
+			if constexpr (bMoveDisabled)
+			{
+				if (autortfm_is_closed())
+				{
+					autortfm_unreachable("TTask functor has an AutoRTFM-disabled move constructor which was called from the closed");
+				}
+			}
 			if constexpr (bCanUseInline)
 			{
 				new(Dst.Inline) T(std::move(*FunctionFrom(Src)));
@@ -108,6 +133,13 @@ private:
 		// Destructs the payload.
 		static void Destruct(FPayload& Payload)
 		{
+			if constexpr (bDestructorDisabled)
+			{
+				if (autortfm_is_closed())
+				{
+					autortfm_unreachable("TTask functor has an AutoRTFM-disabled destructor which was called from the closed");
+				}
+			}
 			if constexpr (bCanUseInline)
 			{
 				FunctionFrom(Payload)->~T();
