[AutoRTFM] Silence PVS false-positive warnings

#Synced-CL 48351162

diff --git a/Engine/Source/Runtime/AutoRTFM/Public/AutoRTFM/Task.h b/Engine/Source/Runtime/AutoRTFM/Public/AutoRTFM/Task.h
--- a/Engine/Source/Runtime/AutoRTFM/Public/AutoRTFM/Task.h
+++ b/Engine/Source/Runtime/AutoRTFM/Public/AutoRTFM/Task.h
@@ -64,13 +64,13 @@ private:
 		static constexpr bool bCanUseInline = sizeof(T) <= InlineDataSize && alignof(T) <= InlineDataAlignment;
 
 		// True iff the copy constructor of the functor T is AutoRTFM-disabled.
-		static constexpr bool bCopyDisabled = AUTORTFM_MODE_OF_CALL(::new T(std::declval<const T&>())) == autortfm_mode_disable;
+		static constexpr bool bCopyDisabled = AUTORTFM_MODE_OF_CALL(::new T(std::declval<const T&>())) == autortfm_mode_disable; //-V501
 
 		// True iff the move constructor of the functor T is AutoRTFM-disabled.
-		static constexpr bool bMoveDisabled = AUTORTFM_MODE_OF_CALL(::new T(std::declval<T&&>())) == autortfm_mode_disable;
+		static constexpr bool bMoveDisabled = AUTORTFM_MODE_OF_CALL(::new T(std::declval<T&&>())) == autortfm_mode_disable; //-V501
 
 		// True iff the destructor of the functor T is AutoRTFM-disabled.
-		static constexpr bool bDestructorDisabled = AUTORTFM_MODE_OF_CALL(std::declval<T>().~T()) == autortfm_mode_disable;
+		static constexpr bool bDestructorDisabled = AUTORTFM_MODE_OF_CALL(std::declval<T>().~T()) == autortfm_mode_disable; //-V501
 
 		// Returns the functor from the payload.
 		static T* FunctionFrom(FPayload& Payload)
