[AutoRTFM] Clarify usage of UE_AUTORTFM_DECLARE_THREAD_LOCAL_VAR macro.

This macro doesn't work at global scope, but this wasn't explained anywhere.

#Synced-CL 48275511

diff --git a/Engine/Source/Runtime/AutoRTFM/Public/AutoRTFM.h b/Engine/Source/Runtime/AutoRTFM/Public/AutoRTFM.h
--- a/Engine/Source/Runtime/AutoRTFM/Public/AutoRTFM.h
+++ b/Engine/Source/Runtime/AutoRTFM/Public/AutoRTFM.h
@@ -1072,9 +1072,9 @@ namespace AutoRTFM::Private
 #endif
 
 // Declares an AutoRTFM-aware thread local variable. `thread_local` variables are not yet natively supported (#jira SOL-7684)
-// Calls should be written like this:
+// This macro must be used inside a function, and will not work properly at global scope. Calls should be written like this:
 //     UE_AUTORTFM_DECLARE_THREAD_LOCAL_VAR(FString, MyThreadLocalString);
-//     MyThreadLocalString = "Hello";
+//     MyThreadLocalString = TEXT("Hello");
 #define UE_AUTORTFM_DECLARE_THREAD_LOCAL_VAR(Type, Name) UE_AUTORTFM_DECLARE_THREAD_LOCAL_VAR_IMPL(Type, Name)
 
 // Runs a block of code in the open, non-transactionally. Anything performed in the open will not be undone if a transaction fails.
