[AutoRTFM] Add AutoRTFM support for custom write rollback handlers

The AutoRTFM API has been expanded to include a "CustomRollback" flag on each write record.
When set, and the transaction is aborted, a new RollbackWrite function will be called to perform the custom rollback logic.

This will be used by VerseVM's TWriteBarrier, so that rollbacks perform a GC barrier on the value whenever the value changes.

#Synced-CL 48603520

diff --git a/Engine/Source/Programs/AutoRTFMTests/Private/WriteLogTests.cpp b/Engine/Source/Programs/AutoRTFMTests/Private/WriteLogTests.cpp
--- a/Engine/Source/Programs/AutoRTFMTests/Private/WriteLogTests.cpp
+++ b/Engine/Source/Programs/AutoRTFMTests/Private/WriteLogTests.cpp
@@ -1,5 +1,6 @@
 // Copyright Epic Games, Inc. All Rights Reserved.
 
+#include "AutoRTFM.h"
 #include "WriteLog.h"
 #include "Catch2Includes.h"
 #include <algorithm>
@@ -32,13 +33,13 @@ TEST_CASE("FWriteLog")
 	};
 
 	// Returns a std::vector<size_t> with NumEntries elements with random values
-	// between 1..AutoRTFM::FWriteLogEntry::MaxSize
-	auto RandomSize = [&Rand](size_t NumEntries)
+	// between 1 and MaxSize
+	auto RandomSize = [&Rand](size_t NumEntries, size_t MaxSize)
 	{
 		std::vector<size_t> EntrySizes(NumEntries);
 		for (size_t& EntrySize : EntrySizes)
 		{
-			EntrySize = std::uniform_int_distribution<unsigned int>(1, AutoRTFM::FWriteLog::RecordMaxSize)(Rand);
+			EntrySize = std::uniform_int_distribution<unsigned int>(1, MaxSize)(Rand);
 		}
 		return EntrySizes;
 	};
@@ -99,11 +100,29 @@ TEST_CASE("FWriteLog")
 		std::vector<AutoRTFM::FWriteLogEntry> Entries;
 		for (FDataSpan& DataSpan : DataSpans)
 		{
-			AutoRTFM::FWriteLogEntry Entry;
+			AutoRTFM::FWriteLogEntry Entry{};
 			Entry.Data = &Data[DataSpan.Offset];
 			Entry.Size = DataSpan.Size;
 			Entry.LogicalAddress = reinterpret_cast<std::byte*>(static_cast<uintptr_t>(0x1234000 + DataSpan.Offset));
-			Entry.bNoMemoryValidation = (static_cast<int>(Entry.Data[0]) & 0x10) != 0;
+			if ((Rand() & 3) == 0) // 25% of write log entries can have non-default flags.
+			{
+				if (Rand() & 1)
+				{
+					Entry.Flags = Entry.Flags | AutoRTFM::EWriteFlags::NoValidation;
+				}
+				if (Rand() & 1)
+				{
+					Entry.Flags = Entry.Flags | AutoRTFM::EWriteFlags::CustomRollback;
+					if (Rand() & 1)
+					{
+						Entry.Flags = Entry.Flags | AutoRTFM::EWriteFlags::CustomFlag0;
+					}
+					if (Rand() & 1)
+					{
+						Entry.Flags = Entry.Flags | AutoRTFM::EWriteFlags::CustomFlag1;
+					}
+				}
+			}
 			Entries.push_back(Entry);
 			WriteLog.Push(Entry);
 		}
@@ -126,9 +145,9 @@ TEST_CASE("FWriteLog")
 						Expect = Entries[EntryIndex++];
 					}
 
-					// The actual write log must match the expectation's address and memory-validation.
+					// The actual write log must match the expectation's address and flags.
 					REQUIRE(Got.LogicalAddress == Expect.LogicalAddress);
-					REQUIRE(Got.bNoMemoryValidation == Expect.bNoMemoryValidation);
+					REQUIRE(Got.Flags == Expect.Flags);
 
 					// Compare the actual contents of the write log with our expectations.
 					const size_t NumBytes = std::min(Got.Size, Expect.Size);
@@ -160,9 +179,9 @@ TEST_CASE("FWriteLog")
 						Expect = Entries[--EntryIndex];
 					}
 
-					// The actual write log must match the expectation's address and memory-validation.
+					// The actual write log must match the expectation's address and flags.
 					REQUIRE(Got.LogicalAddress + Got.Size == Expect.LogicalAddress + Expect.Size);
-					REQUIRE(Got.bNoMemoryValidation == Expect.bNoMemoryValidation);
+					REQUIRE(Got.Flags == Expect.Flags);
 
 					// Compare the actual contents of the write log with our expectations.
 					// Since we are reading the lists in reverse, we check the right edge of the data, and 
@@ -209,12 +228,6 @@ TEST_CASE("FWriteLog")
 		REQUIRE(WriteLog.TotalSize() == 0u);
 	}
 
-
-	SECTION("EntrySize: RecordMaxSize-1, NumEntries: 32")
-	{
-		Check(FixedSize(/* EntrySize */ AutoRTFM::FWriteLog::RecordMaxSize - 1, /* NumEntries */ 32));
-	}
-
 	SECTION("EntrySize: 1, NumEntries: 32")
 	{
 		Check(FixedSize(/* EntrySize */ 1, /* NumEntries */ 32));
@@ -230,24 +243,49 @@ TEST_CASE("FWriteLog")
 		Check(FixedSize(/* EntrySize */ 1024, /* NumEntries */ 32));
 	}
 
-	SECTION("EntrySize: RecordMaxSize, NumEntries: 32")
+	SECTION("EntrySize: FRegularRecordLayout::MaxSize-1, NumEntries: 32")
+	{
+		Check(FixedSize(/* EntrySize */ AutoRTFM::FWriteLog::FRegularRecordLayout::MaxSize - 1, /* NumEntries */ 32));
+	}
+
+	SECTION("EntrySize: FRegularRecordLayout::MaxSize, NumEntries: 32")
+	{
+		Check(FixedSize(/* EntrySize */ AutoRTFM::FWriteLog::FRegularRecordLayout::MaxSize, /* NumEntries */ 32));
+	}
+
+	SECTION("EntrySize: FRegularRecordLayout::MaxSize+1, NumEntries: 32")
+	{
+		Check(FixedSize(/* EntrySize */ AutoRTFM::FWriteLog::FRegularRecordLayout::MaxSize + 1, /* NumEntries */ 32));
+	}
+
+	SECTION("EntrySize: FRegularRecordLayout::MaxSize*10, NumEntries: 10")
+	{
+		Check(FixedSize(/* EntrySize */ AutoRTFM::FWriteLog::FRegularRecordLayout::MaxSize * 10, /* NumEntries */ 10));
+	}
+
+	SECTION("EntrySize: FCustomRollbackRecordLayout::MaxSize-1, NumEntries: 32")
+	{
+		Check(FixedSize(/* EntrySize */ AutoRTFM::FWriteLog::FCustomRollbackRecordLayout::MaxSize - 1, /* NumEntries */ 32));
+	}
+
+	SECTION("EntrySize: FCustomRollbackRecordLayout::MaxSize, NumEntries: 32")
 	{
-		Check(FixedSize(/* EntrySize */ AutoRTFM::FWriteLog::RecordMaxSize, /* NumEntries */ 32));
+		Check(FixedSize(/* EntrySize */ AutoRTFM::FWriteLog::FCustomRollbackRecordLayout::MaxSize, /* NumEntries */ 32));
 	}
 
-	SECTION("EntrySize: RecordMaxSize+1, NumEntries: 32")
+	SECTION("EntrySize: FCustomRollbackRecordLayout::MaxSize+1, NumEntries: 32")
 	{
-		Check(FixedSize(/* EntrySize */ AutoRTFM::FWriteLog::RecordMaxSize+1, /* NumEntries */ 32));
+		Check(FixedSize(/* EntrySize */ AutoRTFM::FWriteLog::FCustomRollbackRecordLayout::MaxSize + 1, /* NumEntries */ 32));
 	}
 
-	SECTION("EntrySize: RecordMaxSize*10, NumEntries: 10")
+	SECTION("EntrySize: FCustomRollbackRecordLayout::MaxSize*10, NumEntries: 10")
 	{
-		Check(FixedSize(/* EntrySize */ AutoRTFM::FWriteLog::RecordMaxSize * 10, /* NumEntries */ 10));
+		Check(FixedSize(/* EntrySize */ AutoRTFM::FWriteLog::FCustomRollbackRecordLayout::MaxSize * 10, /* NumEntries */ 10));
 	}
 
 	SECTION("EntrySize: random, NumEntries: 32")
 	{
-		Check(RandomSize(/* NumEntries */ 32));
+		Check(RandomSize(/* NumEntries */ 32, /* MaxSize */ 0x8000));
 	}
 }
 
@@ -291,7 +329,7 @@ TEST_CASE("FWriteLog.Hash")
 				.LogicalAddress = Buffer.Data + Offset,
 				.Data = Buffer.Data + Offset,
 				.Size = WriteSize,
-				.bNoMemoryValidation = false
+				.Flags = AutoRTFM::EWriteFlags::Default
 			});
 
 			Writes.push_back(FMemorySpan{Buffer.Data + Offset, WriteSize});
diff --git a/Engine/Source/Runtime/AutoRTFM/Private/API.cpp b/Engine/Source/Runtime/AutoRTFM/Private/API.cpp
--- a/Engine/Source/Runtime/AutoRTFM/Private/API.cpp
+++ b/Engine/Source/Runtime/AutoRTFM/Private/API.cpp
@@ -971,7 +971,7 @@ extern "C" void autortfm_record_open_write(void* Ptr, size_t Size) noexcept
 	}
 }
 
-extern "C" void autortfm_record_open_write_no_memory_validation(void* Ptr, size_t Size) noexcept
+extern "C" void autortfm_record_open_write_with_flags(void* Ptr, size_t Size, autortfm_write_flags Flags) noexcept
 {
 	if (autortfm_is_closed())
 	{
@@ -981,7 +981,7 @@ extern "C" void autortfm_record_open_write_no_memory_validation(void* Ptr, size_
 	{
 		if (FTransaction* CurrentTransaction = Context->GetCurrentTransaction())
 		{
-			CurrentTransaction->RecordWrite(Ptr, Size, /* bNoMemoryValidation */ true);
+			CurrentTransaction->RecordWrite(Ptr, Size, static_cast<EWriteFlags>(Flags));
 		}
 	}
 }
diff --git a/Engine/Source/Runtime/AutoRTFM/Private/ExternAPI.cpp b/Engine/Source/Runtime/AutoRTFM/Private/ExternAPI.cpp
--- a/Engine/Source/Runtime/AutoRTFM/Private/ExternAPI.cpp
+++ b/Engine/Source/Runtime/AutoRTFM/Private/ExternAPI.cpp
@@ -22,18 +22,18 @@ namespace
 // crashing due to a nullptr.
 void DefaultLog(const char* File, int Line, [[maybe_unused]] void* ProgramCounter, autortfm_log_severity Severity, const char* Format, va_list Args)
 {
-	static constexpr size_t InlineBufferLength = 256;
+	static constexpr int InlineBufferLength = 256;
 	char InlineBuffer[InlineBufferLength];
 
 	va_list Args2;
 	va_copy(Args2, Args);
 	int Count = vsnprintf(InlineBuffer, InlineBufferLength, Format, Args);
-	auto AllocatedBuffer = (static_cast<size_t>(Count) >= InlineBufferLength) ? std::unique_ptr<char[]>{new char[Count+1]} : nullptr;
+	auto AllocatedBuffer = (Count >= InlineBufferLength) ? std::unique_ptr<char[]>{new char[Count+1]} : nullptr;
 	if (AllocatedBuffer)
 	{
 		vsnprintf(AllocatedBuffer.get(), Count+1, Format, Args2);
-		va_end(Args2);
 	}
+	va_end(Args2);
 
 	char* const Buffer = AllocatedBuffer ? AllocatedBuffer.get() : InlineBuffer;
 
@@ -76,6 +76,7 @@ autortfm_extern_api GExternAPI
 	/* OnMemoryValidationThrottlingChanged */ nullptr,
 	/* OnMemoryValidationWriteIdentifierChanged */ nullptr,
 	/* OnMemoryValidationStatisticsChanged */ nullptr,
+	/* RollbackWrite */ nullptr,
 };
 
 } // namespace AutoRTFM
diff --git a/Engine/Source/Runtime/AutoRTFM/Private/HitSet.h b/Engine/Source/Runtime/AutoRTFM/Private/HitSet.h
--- a/Engine/Source/Runtime/AutoRTFM/Private/HitSet.h
+++ b/Engine/Source/Runtime/AutoRTFM/Private/HitSet.h
@@ -16,9 +16,9 @@ namespace AutoRTFM
 		// The address of the write.
 		uintptr_t Address : 48;
 		// Size of the write.
-		uintptr_t Size : 15;
-		// If set, then the write shouldn't be considered by the memory validator.
-		uintptr_t bNoMemoryValidation : 1;
+		uintptr_t Size : 12;
+		// Flags applied to the write.
+		uintptr_t Flags : 4;
 		
 		UE_AUTORTFM_FORCEINLINE uintptr_t Payload() const
 		{
diff --git a/Engine/Source/Runtime/AutoRTFM/Private/Transaction.cpp b/Engine/Source/Runtime/AutoRTFM/Private/Transaction.cpp
--- a/Engine/Source/Runtime/AutoRTFM/Private/Transaction.cpp
+++ b/Engine/Source/Runtime/AutoRTFM/Private/Transaction.cpp
@@ -331,7 +331,14 @@ void FTransaction::Undo()
 		// No write records should be within the transaction's stack range.
 		AUTORTFM_ENSURE(!IsOnStack(Entry.LogicalAddress));
 
-		memcpy(Entry.LogicalAddress, Entry.Data, Entry.Size);
+		if ((Entry.Flags & EWriteFlags::CustomRollback) == EWriteFlags::CustomRollback)
+		{
+			GExternAPI.RollbackWrite(Entry.LogicalAddress, Entry.Data, Entry.Size, static_cast<autortfm_write_flags>(Entry.Flags));
+		}
+		else
+		{
+			memcpy(Entry.LogicalAddress, Entry.Data, Entry.Size);
+		}
     }
 
 	AUTORTFM_VERBOSE("Undone a transaction!");
@@ -357,7 +364,7 @@ void FTransaction::CommitNested()
 			FHitSetEntry HitSetEntry{};
 			HitSetEntry.Address = reinterpret_cast<uintptr_t>(Write.LogicalAddress);
 			HitSetEntry.Size = static_cast<uint16_t>(Write.Size);
-			HitSetEntry.bNoMemoryValidation = Write.bNoMemoryValidation;
+			HitSetEntry.Flags = static_cast<uintptr_t>(Write.Flags);
 
 			if (Parent->HitSet.FindOrTryInsert(HitSetEntry))
 			{
diff --git a/Engine/Source/Runtime/AutoRTFM/Private/Transaction.h b/Engine/Source/Runtime/AutoRTFM/Private/Transaction.h
--- a/Engine/Source/Runtime/AutoRTFM/Private/Transaction.h
+++ b/Engine/Source/Runtime/AutoRTFM/Private/Transaction.h
@@ -117,7 +117,7 @@ public:
 	bool AttemptToCommit();
 
 	// Record that a write is about to occur at the given LogicalAddress of Size bytes.
-	void RecordWrite(void* LogicalAddress, size_t Size, bool bNoMemoryValidation = false);
+	void RecordWrite(void* LogicalAddress, size_t Size, EWriteFlags Flags = EWriteFlags::Default);
 	template<unsigned SIZE> void RecordWrite(void* LogicalAddress);
 	template<unsigned SIZE> void RecordWriteInsertedSlow(void* LogicalAddress);
 	template<unsigned SIZE> void RecordWriteNotInsertedSlow(void* LogicalAddress);
diff --git a/Engine/Source/Runtime/AutoRTFM/Private/TransactionInlines.h b/Engine/Source/Runtime/AutoRTFM/Private/TransactionInlines.h
--- a/Engine/Source/Runtime/AutoRTFM/Private/TransactionInlines.h
+++ b/Engine/Source/Runtime/AutoRTFM/Private/TransactionInlines.h
@@ -62,7 +62,7 @@ UE_AUTORTFM_FORCEINLINE bool FTransaction::ShouldRecordWrite(void* LogicalAddres
 }
 
 AUTORTFM_NO_ASAN UE_AUTORTFM_FORCEINLINE
-void FTransaction::RecordWrite(void* LogicalAddress, size_t Size, bool bNoMemoryValidation /* = false */)
+void FTransaction::RecordWrite(void* LogicalAddress, size_t Size, EWriteFlags Flags /* = EWriteFlags::Default */)
 {
 	if (AUTORTFM_UNLIKELY(0 == Size))
 	{
@@ -80,7 +80,7 @@ void FTransaction::RecordWrite(void* LogicalAddress, size_t Size, bool bNoMemory
 		FHitSetEntry HitSetEntry{};
 		HitSetEntry.Address = reinterpret_cast<uintptr_t>(LogicalAddress);
 		HitSetEntry.Size = static_cast<uint16_t>(Size);
-		HitSetEntry.bNoMemoryValidation = bNoMemoryValidation;
+		HitSetEntry.Flags = static_cast<uintptr_t>(Flags);
 
 		if (HitSet.FindOrTryInsert(HitSetEntry))
 		{
@@ -99,13 +99,30 @@ void FTransaction::RecordWrite(void* LogicalAddress, size_t Size, bool bNoMemory
 
 	Stats.Collect<EStatsKind::NewMemoryTrackerMiss>();
 
-	WriteLog.Push(FWriteLogEntry
+	FWriteLogEntry Entry
 	{
 		.LogicalAddress = static_cast<std::byte*>(LogicalAddress),
 		.Data = static_cast<std::byte*>(LogicalAddress),
 		.Size = Size,
-		.bNoMemoryValidation = bNoMemoryValidation
-	});
+		.Flags = Flags
+	};
+
+	switch (Entry.Size)
+	{
+		case 8:
+			WriteLog.PushSmall<8>(Entry.LogicalAddress, Entry.Data, Entry.Flags);
+			break;
+		case 4:
+			WriteLog.PushSmall<4>(Entry.LogicalAddress, Entry.Data, Entry.Flags);
+			break;
+		case 2:
+			WriteLog.PushSmall<2>(Entry.LogicalAddress, Entry.Data, Entry.Flags);
+			break;
+		default:
+			WriteLog.Push(Entry);
+			break;
+	}
+	
 }
 
 template<unsigned SIZE> AUTORTFM_NO_ASAN UE_AUTORTFM_FORCEINLINE void FTransaction::RecordWrite(void* LogicalAddress)
@@ -161,7 +178,7 @@ template<unsigned SIZE> AUTORTFM_NO_ASAN UE_AUTORTFM_FORCENOINLINE void FTransac
 
 	Stats.Collect<EStatsKind::NewMemoryTrackerMiss>();
 
-	WriteLog.PushSmall<SIZE>(static_cast<std::byte*>(LogicalAddress));
+	WriteLog.PushSmall<SIZE>(static_cast<std::byte*>(LogicalAddress), static_cast<std::byte*>(LogicalAddress));
 }
 
 UE_AUTORTFM_FORCEINLINE void FTransaction::DidAllocate(void* LogicalAddress, const size_t Size)
diff --git a/Engine/Source/Runtime/AutoRTFM/Private/WriteLog.cpp b/Engine/Source/Runtime/AutoRTFM/Private/WriteLog.cpp
--- a/Engine/Source/Runtime/AutoRTFM/Private/WriteLog.cpp
+++ b/Engine/Source/Runtime/AutoRTFM/Private/WriteLog.cpp
@@ -4,6 +4,7 @@
 
 #include "WriteLog.h"
 
+#include "AutoRTFM.h"
 #include "BuildMacros.h"
 #include "Utils.h"
 
@@ -89,7 +90,7 @@ FWriteLog::FHash FWriteLog::Hash(size_t NumWriteEntries) const
 		}
 
 		const FWriteLogEntry& Entry = *Iter;
-		if (AUTORTFM_UNLIKELY(Entry.bNoMemoryValidation))
+		if (AUTORTFM_UNLIKELY((Entry.Flags & EWriteFlags::NoValidation) == EWriteFlags::NoValidation))
 		{
 			continue; // Next write
 		}
@@ -144,7 +145,7 @@ FWriteLog::FHash FWriteLog::HashAVX2(size_t NumWriteEntries) const
 		}
 
 		const FWriteLogEntry& Entry = *Iter;
-		if (AUTORTFM_UNLIKELY(Entry.bNoMemoryValidation))
+		if (AUTORTFM_UNLIKELY((Entry.Flags & EWriteFlags::NoValidation) == EWriteFlags::NoValidation))
 		{
 			continue; // Next write
 		}
diff --git a/Engine/Source/Runtime/AutoRTFM/Private/WriteLog.h b/Engine/Source/Runtime/AutoRTFM/Private/WriteLog.h
--- a/Engine/Source/Runtime/AutoRTFM/Private/WriteLog.h
+++ b/Engine/Source/Runtime/AutoRTFM/Private/WriteLog.h
@@ -27,58 +27,157 @@ namespace AutoRTFM
 		// The size of the write in bytes.
 		size_t Size = 0;
 		
-		// If true, then this write will not be considered by the AutoRTFM
-		// memory validator.
-		bool bNoMemoryValidation = false;
+		// Flags on the write.
+		EWriteFlags Flags = EWriteFlags::Default;
+
+		// Returns a new FWriteLogEntry with the logical address and data
+		// pointers increased by NumBytes, the size decreased by NumBytes, and
+		// the same flags.
+		[[nodiscard]] FWriteLogEntry Advance(size_t NumBytes) const
+		{
+			return {
+				.LogicalAddress = LogicalAddress + NumBytes,
+				.Data = Data + NumBytes,
+				.Size = Size - NumBytes,
+				.Flags = Flags,
+			};
+		}
 	};
 
 	// FWriteLog holds an ordered list of write records which can be iterated
 	// forwards and backwards.
-	// Ensure changes to this class are kept in sync with Unreal.natvis.
+	// Ensure changes to this class are kept in sync with AutoRTFM.natvis.
 	class FWriteLog final
 	{
+		enum EAppendMode
+		{
+			Partial,
+			Whole,
+		};
 	public:
-		// Number of bits used by the FRecord to represent a write's size.
-		static constexpr size_t RecordSizeBits = 15;
+		template<size_t SizeBits, size_t FlagBits, size_t AddressBits>
+		struct FRecordLayout
+		{
+			// Number of bits used by the FRecord to represent a write's address.
+			static constexpr size_t NumAddressBits = AddressBits;
+
+			// Number of bits used by the FRecord to represent a write's flags.
+			static constexpr size_t NumFlagsBits = FlagBits;
+
+			// Number of bits used by the FRecord to represent a write's size.
+			// Note a record size must be greater than one byte, so is stored as Size-1.
+			static constexpr size_t NumSizeBits = SizeBits;
+			
+			static_assert(NumAddressBits + NumFlagsBits + NumSizeBits == 64);
 
-		// The maximum size for a single write log entry.
-		// Writes will be split into multiple entries if the write is too large.
-		static constexpr size_t RecordMaxSize = (1u << RecordSizeBits) - 1;
+			static constexpr size_t AddressShift = 0;
+			static constexpr size_t AddressMask = (static_cast<size_t>(1) << NumAddressBits) - 1;
+			static constexpr size_t FlagsShift = NumAddressBits;
+			static constexpr size_t FlagsMask = (static_cast<size_t>(1) << NumFlagsBits) - 1;
+			static constexpr size_t SizeShift = NumAddressBits + NumFlagsBits;
+			static constexpr size_t SizeMask = (static_cast<size_t>(1) << NumSizeBits) - 1;
+
+			static constexpr size_t MaxAddress = AddressMask;
+			static constexpr size_t MaxSize = SizeMask + 1;
+
+			[[nodiscard]] static uint64_t Pack(std::byte* Address, EWriteFlags Flags, size_t Size)
+			{
+				const size_t SizeM1 = Size - 1;
+				AUTORTFM_ASSERT_DEBUG((static_cast<size_t>(reinterpret_cast<uintptr_t>(Address)) & ~AddressMask) == 0);
+				AUTORTFM_ASSERT_DEBUG((static_cast<size_t>(Flags) & ~FlagsMask) == 0);
+				AUTORTFM_ASSERT_DEBUG((SizeM1 & ~SizeMask) == 0);
+
+				return reinterpret_cast<uint64_t>(Address)
+					| (static_cast<uint64_t>(Flags) << FlagsShift)
+					| (static_cast<uint64_t>(SizeM1) << SizeShift);
+			}
+
+			[[nodiscard]] static uint64_t ReplaceSize(uint64_t Packed, size_t NewSize)
+			{
+				const size_t NewSizeM1 = NewSize - 1;
+				AUTORTFM_ASSERT_DEBUG((NewSizeM1 & ~SizeMask) == 0);
+				constexpr uint64_t Mask = ~(SizeMask << SizeShift);
+				return (Packed & Mask) | (NewSizeM1 << SizeShift);
+			}
+
+			[[nodiscard]] static std::byte* UnpackAddress(uint64_t Packed)
+			{
+				return reinterpret_cast<std::byte*>((Packed >> AddressShift) & AddressMask);
+			}
+
+			[[nodiscard]] static EWriteFlags UnpackFlags(uint64_t Packed)
+			{
+				return static_cast<EWriteFlags>((Packed >> NumAddressBits) & FlagsMask);
+			}
+
+			[[nodiscard]] static size_t UnpackSize(uint64_t Packed)
+			{
+				const size_t SizeM1 = reinterpret_cast<size_t>((Packed >> SizeShift) & SizeMask);
+				return SizeM1 + 1;
+			}
+		};
+
+		//  64           50 48                                              0
+		//  ┌─────────────┬─┬───────────────────────────────────────────────┐
+		//  │    Size     │F│                    Address                    │
+		//  └─────────────┴─┴───────────────────────────────────────────────┘
+		using FRegularRecordLayout = FRecordLayout</* SizeBits */ 14, /* FlagBits */ 2, /* AddressBits */ 48>;
+
+		//  64         52   48                                              0
+		//  ┌───────────┬───┬───────────────────────────────────────────────┐
+		//  │   Size    │ F │                    Address                    │
+		//  └───────────┴───┴───────────────────────────────────────────────┘
+		using FCustomRollbackRecordLayout = FRecordLayout</* SizeBits */ 12, /* FlagBits */ 4, /* AddressBits */ 48>;
 
 	private:
 		struct FRecord
 		{
-			void Set(std::byte* Address, size_t Size, bool bNoMemoryValidation)
+			// Returns true if the record has the EWriteFlags::CustomRollback flag set.
+			// This alters the bit packing layout of the record.
+			bool UsesCustomRollback() const
 			{
-				AUTORTFM_ASSERT_DEBUG((reinterpret_cast<uintptr_t>(Address) & 0xFFFF0000'00000000) == 0);
-				AUTORTFM_ASSERT_DEBUG(Size <= 0x7FFF);
-		
-				Bits = reinterpret_cast<uint64_t>(Address);
-				Bits <<= 1;
-				Bits |= bNoMemoryValidation;
-				Bits <<= 15;
-				Bits |= Size;
+				return (FRegularRecordLayout::UnpackFlags(Bits) & EWriteFlags::CustomRollback) == EWriteFlags::CustomRollback;
+			}
+
+			// Packs the address, flags and size into the record using the provided bit packing layout.
+			template <typename Layout>
+			void Set(std::byte* Address, EWriteFlags Flags, size_t Size)
+			{
+				AUTORTFM_ASSERT_DEBUG(Size > 0);
+				Bits = Layout::Pack(Address, Flags, Size);
 			}
 
-			void Grow(size_t Amount)
+			// Replaces the size field of the record using the provided bit packing layout.
+			template <typename Layout>
+			void SetSize(size_t NewSize)
 			{
-				AUTORTFM_ASSERT_DEBUG(Size() + Amount <= 0x7FFF);
-				Bits += Amount;
+				AUTORTFM_ASSERT_DEBUG(NewSize > 0);
+				Bits = Layout::ReplaceSize(Bits, NewSize);
 			}
 
+			// Returns the unpacked address from the record.
 			std::byte* Address() const
 			{
-				return reinterpret_cast<std::byte*>(Bits >> 16);
+				return UsesCustomRollback() ? FCustomRollbackRecordLayout::UnpackAddress(Bits) : FRegularRecordLayout::UnpackAddress(Bits);
 			}
 
-			bool NoMemoryValidation() const
+			// Returns the unpacked flags from the record.
+			EWriteFlags Flags() const
 			{
-				return Bits & 0x8000;
+				return UsesCustomRollback() ? FCustomRollbackRecordLayout::UnpackFlags(Bits) : FRegularRecordLayout::UnpackFlags(Bits);
 			}
 
+			// Returns the unpacked size from the record.
 			size_t Size() const
 			{
-				return Bits & 0x7FFF;
+				return UsesCustomRollback() ? FCustomRollbackRecordLayout::UnpackSize(Bits) : FRegularRecordLayout::UnpackSize(Bits);
+			}
+
+			// Returns the unpacked size from the record, using the bit packing layout.
+			template <typename Layout>
+			size_t Size() const
+			{
+				return Layout::UnpackSize(Bits);
 			}
 
 			uint64_t Bits;
@@ -152,78 +251,110 @@ namespace AutoRTFM
 				return LastRecord > FirstRecord;
 			}
 
-			// This will return 0 if the passed-in entry can be folded into the previous write, and 1 if a new
-			// write record will be required.
-			UE_AUTORTFM_FORCEINLINE int NumRecordsNeededForWrite(std::byte* LogicalAddress, bool bNoMemoryValidation)
+			// Returns the number of unused bytes in the block between the last
+			// record and end of the data.
+			UE_AUTORTFM_FORCEINLINE size_t SpaceRemaining() const
 			{
-				return static_cast<int>(
-					IsEmpty() ||
-					LogicalAddress != LastRecord->Address() + LastRecord->Size() ||
-					LastRecord->NoMemoryValidation() != bNoMemoryValidation);
+				return reinterpret_cast<std::byte*>(LastRecord) - DataEnd;
 			}
 
-			// This determines the number of bytes that can safely be passed to `Push` below, given the number of
-			// bytes in the write entry (`EntrySize`) and whether or not those bytes can be folded into an existing
-			// record (which can be determined by calling `NumRecordsNeededForWrite` above). This algorithm does
-			// _not_ clamp the input size to RecordMaxSize. Instead, the block allocation logic in `Push` and
-			// `PushSmall` avoids creating an FBlock larger than MaxSize at all. (When very large pushes occur, the
-			// `Push` algorithm _will_ create a large FBlock to satisfy it, but this logic does not call
-			// `CalculatePushBytes`, and will create a new tail block afterwards.)
-			// 
-			// Returns zero if the block is entirely full, `EntrySize` if there's enough space, or a value
-			// somewhere in-between if the entry needs to be split.
-			UE_AUTORTFM_FORCEINLINE size_t CalculatePushBytes(size_t EntrySize, int NumRecordsNeeded)
+			// Appends a new record to the block, returning the number of bytes
+			// copied to the new record from the front of DataIn.
+			//
+			// If AppendMode is Partial, then AppendWithNewRecord() will append
+			// as many bytes that will fit in the remaining available space of
+			// the block, which may be fewer than NumBytes. In this mode
+			// NumBytes may be 0, in which case the function does nothing and
+			// returns 0.
+			//
+			// If AppendMode is Whole, then it is the caller's responsibility
+			// to ensure that there is enough space in the block to hold an
+			// additional NumBytes of data plus the size of an FRecord. In this
+			// mode NumBytes must be greater than zero.
+			//
+			// Layout is the bit-packing layout of a record with the given flags.
+			template<EAppendMode AppendMode, typename Layout>
+			UE_AUTORTFM_FORCEINLINE size_t AppendWithNewRecord(std::byte* LogicalAddress, std::byte* DataIn, size_t NumBytes, EWriteFlags Flags)
 			{
-				std::byte* BlockEdge = reinterpret_cast<std::byte*>(LastRecord - NumRecordsNeeded);
-				constexpr ptrdiff_t MinimumWorthwhileSplitSize = 8;
-
-				if (DataEnd + EntrySize <= BlockEdge)
-				{
-					// We will fit the entire entry in the block with room to spare.
-					AUTORTFM_ASSERT_DEBUG(EntrySize <= RecordMaxSize);
-					AUTORTFM_ASSERT_DEBUG((NumRecordsNeeded == 1) || (LastRecord->Size() + EntrySize <= ptrdiff_t(RecordMaxSize)));
-					return EntrySize;
-				}
-				else if (ptrdiff_t BlockCapacity = BlockEdge - DataEnd; BlockCapacity >= MinimumWorthwhileSplitSize)
+				if constexpr (AppendMode == Partial)
 				{
-					// We have enough data to fill up the entire block. This path will split the data across the end of this
-					// block and the start of the next block. We avoid this path when the block has fewer than eight bytes
-					// left, just as an efficiency measure, since it takes extra time to assemble two FRecords instead of one.
-					AUTORTFM_ASSERT_DEBUG(BlockCapacity <= ptrdiff_t(RecordMaxSize));
-					AUTORTFM_ASSERT_DEBUG((NumRecordsNeeded == 1) || (LastRecord->Size() + BlockCapacity <= ptrdiff_t(RecordMaxSize)));
-					return BlockCapacity;
+					const size_t Space = SpaceRemaining();
+					if (NumBytes == 0 || Space <= sizeof(FRecord))
+					{
+						return 0;
+					}
+					NumBytes = std::min(std::min(NumBytes, Space - sizeof(FRecord)), Layout::MaxSize);
 				}
 				else
 				{
-					// The block is completely full.
-					return 0;
+					static_assert(AppendMode == Whole);
+					AUTORTFM_ASSERT_DEBUG(NumBytes > 0);
 				}
+
+				AUTORTFM_ASSERT_DEBUG(SpaceRemaining() >= sizeof(FRecord) + NumBytes);
+
+				LastRecord--;
+				LastRecord->Set<Layout>(LogicalAddress, Flags, NumBytes);
+				memcpy(DataEnd, DataIn, NumBytes);
+				DataEnd += NumBytes;
+				AUTORTFM_ASSERT_DEBUG(DataEnd <= reinterpret_cast<std::byte*>(LastRecord));
+
+				return NumBytes;
 			}
 
-			// Grows this block by copying `NumBytes` bytes from `DataIn`, representing data originally from `LogicalAddress`. 
-			// Asserts if the block does not have enough capacity to hold `NumBytes`. The caller should determine the available
-			// capacity ahead of time by calling `NumRecordsNeededForWrite` and `CalculatePushBytes`.
-			UE_AUTORTFM_FORCEINLINE void Push(std::byte* LogicalAddress, std::byte* DataIn, size_t NumBytes, bool bNoMemoryValidation, int NumRecordsNeeded)
+			// Appends additional data to the tail record, returning the number
+			// of bytes copied to the end of the record from the front of DataIn.
+			//
+			// Regardless of AppendMode, if the last record cannot be appended
+			// to, either because the logical address does not start at the end
+			// of the last record's logical address or the flags are
+			// incompatible, then the function returns 0.
+			//
+			// If AppendMode is Partial, then AppendWithNewRecord() will append
+			// as many bytes that will fit in the remaining available space of
+			// the block, which may be fewer than NumBytes.
+			//
+			// If AppendMode is Whole, then AppendWithNewRecord() will only
+			// append the full NumBytes if they fit, otherwise the function will
+			// do nothing and return 0.
+			//
+			// Layout is the bit-packing layout of a record with the given flags.
+			template<EAppendMode AppendMode, typename Layout>
+			UE_AUTORTFM_FORCEINLINE size_t AppendToLastRecord(std::byte* LogicalAddress, std::byte* DataIn, size_t NumBytes, EWriteFlags Flags)
 			{
-				AUTORTFM_ASSERT_DEBUG(DataEnd + NumBytes <= reinterpret_cast<std::byte*>(LastRecord - NumRecordsNeeded));
-				AUTORTFM_ASSERT_DEBUG(NumBytes <= RecordMaxSize);
+				if (IsEmpty())
+				{
+					return 0;
+				}
+				
+				const size_t LastRecordSize = LastRecord->Size();
+				if (LogicalAddress != LastRecord->Address() + LastRecordSize ||
+					LastRecord->Flags() != Flags)
+				{
+					return 0;
+				}
+
+				const size_t RepresentableSizeRemaining = Layout::MaxSize - LastRecordSize;
+				const size_t Available = std::min(RepresentableSizeRemaining, SpaceRemaining());
 
-				if (NumRecordsNeeded == 1)
+				if constexpr (AppendMode == Whole)
 				{
-					LastRecord--;
-					LastRecord->Set(LogicalAddress, NumBytes, bNoMemoryValidation);
+					if (NumBytes > Available)
+					{
+						return 0;
+					}
 				}
 				else
 				{
-					AUTORTFM_ASSERT_DEBUG(NumRecordsNeeded == 0);
-					LastRecord->Grow(NumBytes);
+					static_assert(AppendMode == Partial);
+					NumBytes = std::min(NumBytes, Available);
 				}
 
+				LastRecord->SetSize<Layout>(LastRecordSize + NumBytes);
 				memcpy(DataEnd, DataIn, NumBytes);
-
 				DataEnd += NumBytes;
-
 				AUTORTFM_ASSERT_DEBUG(DataEnd <= reinterpret_cast<std::byte*>(LastRecord));
+				return NumBytes;
 			}
 
 			// The next block in the linked list.
@@ -257,107 +388,29 @@ namespace AutoRTFM
 		// The log will make a copy of the FWriteLogEntry's data.
 		void Push(FWriteLogEntry Entry)
 		{
-			AUTORTFM_ASSERT_DEBUG((reinterpret_cast<uintptr_t>(Entry.LogicalAddress) & 0xffff0000'00000000) == 0);
-
+			if ((Entry.Flags & EWriteFlags::CustomRollback) == EWriteFlags::CustomRollback)
 			{
-				TotalSizeBytes += Entry.Size;
-				int NumRecordsNeeded = TailBlock->NumRecordsNeededForWrite(Entry.LogicalAddress, Entry.bNoMemoryValidation);
-				size_t NumBytes = TailBlock->CalculatePushBytes(Entry.Size, NumRecordsNeeded);
-
-				if (NumBytes == Entry.Size)
-				{
-					// This push fits into our existing block.
-					TailBlock->Push(Entry.LogicalAddress, Entry.Data, NumBytes, Entry.bNoMemoryValidation, NumRecordsNeeded);
-					NumEntries += NumRecordsNeeded;
-					return;
-				}
-
-				// The push doesn't fit into the existing block...
-				if (NumBytes > 0)
-				{
-					// ... but we can still use up the remainder of the block.
-					TailBlock->Push(Entry.LogicalAddress, Entry.Data, NumBytes, Entry.bNoMemoryValidation, NumRecordsNeeded);
-					NumEntries += NumRecordsNeeded;
-
-					// Adjust the entry to point to the remaining unlogged bytes.
-					Entry.Size -= NumBytes;
-					Entry.LogicalAddress += NumBytes;
-					Entry.Data += NumBytes;
-				}
-			}
-
-			// Calculate how many maxed-out RecordMaxSize entries we can make.
-			const size_t NumFullRecords = Entry.Size / RecordMaxSize;
-			// Calculate how many bytes will be remaining once we have emitted the full-size entries.
-			const size_t RemainingBytes = Entry.Size - (NumFullRecords * RecordMaxSize);
-			// Calculate the exact required size of the block.
-			const size_t RequiredSize =
-				sizeof(FBlock) +                                                   // FBlock header
-				(NumFullRecords * (RecordMaxSize + sizeof(FRecord))) +             // Bytes needed for full records
-				(LIKELY(RemainingBytes) ? (RemainingBytes + sizeof(FRecord)) : 0); // Bytes needed for trailing partial record, if any
-			// Add padding to the block to account for alignment.
-			const size_t AlignedSize = AlignUp(RequiredSize, sizeof(FRecord));
-
-			// Create a new empty tail block, large enough to hold the remainder of this push in its entirety, and
-			// never smaller than the upcoming block size.
-			const size_t BlockSize = std::max(AlignedSize, NextBlockSize);
-			AllocateNewBlock(BlockSize);
-
-			// Push all of the full records.
-			for (size_t Index = NumFullRecords; Index--; )
-			{
-				constexpr int NumRecordsNeeded = 1;
-				constexpr size_t NumBytes = RecordMaxSize;
-
-				TailBlock->Push(Entry.LogicalAddress, Entry.Data, NumBytes, Entry.bNoMemoryValidation, NumRecordsNeeded);
-				NumEntries += NumRecordsNeeded;
-
-				Entry.Size -= NumBytes;
-				Entry.LogicalAddress += NumBytes;
-				Entry.Data += NumBytes;
-			}
-
-			// Push the final, partial record.
-			if (LIKELY(RemainingBytes))
-			{
-				constexpr int NumRecordsNeeded = 1;
-
-				TailBlock->Push(Entry.LogicalAddress, Entry.Data, RemainingBytes, Entry.bNoMemoryValidation, NumRecordsNeeded);
-				NumEntries += NumRecordsNeeded;
+				Push<FCustomRollbackRecordLayout>(Entry);
 			}
-
-			// If we've just created an extra-large block, and alignment has caused it to be less than 100% full,
-			// we preemptively allocate a new block here. This avoids a scenario where a future PushSmall could
-			// accidentally overflow the record size. Normally this would be impossible because BlockMaxSize
-			// isn't large enough to hold more than RecordMaxSize bytes, but in this case we are making an FBlock
-			// which is potentially much larger than normal.
-			if (AlignedSize > BlockMaxSize && AlignedSize > RequiredSize)
+			else
 			{
-				AllocateNewBlock(NextBlockSize);
+				Push<FRegularRecordLayout>(Entry);
 			}
 		}
 
 		// Adds the write log entry to the log; assumes a payload small enough that splitting is not beneficial.
 		// If you have large sizes, you should use `Push` instead so that splitting can occur.
-		template <unsigned int SIZE>
-		void PushSmall(std::byte* LogicalAddress)
+		template <unsigned int Size>
+		void PushSmall(std::byte* LogicalAddress, std::byte* Data, EWriteFlags Flags = EWriteFlags::Default)
 		{
-			static_assert(SIZE <= RecordMaxSize);  // This is a hard upper limit.
-			AUTORTFM_ASSERT_DEBUG((reinterpret_cast<uintptr_t>(LogicalAddress) & 0xffff0000'00000000) == 0);
-
-			TotalSizeBytes += SIZE;
-
-			int NumRecordsNeeded = TailBlock->NumRecordsNeededForWrite(LogicalAddress, /*bNoMemoryValidation=*/false);
-			size_t NumBytes = TailBlock->CalculatePushBytes(SIZE, NumRecordsNeeded);
-
-			if (UNLIKELY(NumBytes != SIZE))
+			if ((Flags & EWriteFlags::CustomRollback) == EWriteFlags::CustomRollback)
 			{
-				AllocateNewBlock(NextBlockSize);
-				NumRecordsNeeded = 1;
+				PushSmall<Size, FCustomRollbackRecordLayout>(LogicalAddress, Data, Flags);
+			}
+			else
+			{
+				PushSmall<Size, FRegularRecordLayout>(LogicalAddress, Data, Flags);
 			}
-
-			TailBlock->Push(LogicalAddress, LogicalAddress, SIZE, /*bNoMemoryValidation=*/false, NumRecordsNeeded);
-			NumEntries += NumRecordsNeeded;
 		}
 
 		// Iterator for enumerating the writes of the log.
@@ -389,7 +442,7 @@ namespace AutoRTFM
 					.LogicalAddress = reinterpret_cast<std::byte*>(Record->Address()),
 					.Data = Data,
 					.Size = Record->Size(),
-					.bNoMemoryValidation = Record->NoMemoryValidation(),
+					.Flags = Record->Flags(),
 				};
 			}
 
@@ -519,15 +572,94 @@ namespace AutoRTFM
 		FHash Hash(size_t NumWriteEntries) const;
 
 	private:
+		template<typename Layout>
+		void Push(FWriteLogEntry Entry)
+		{
+			TotalSizeBytes += Entry.Size;
+
+			// Attempt to append to the end of the last record.
+			if (const size_t Appended = TailBlock->AppendToLastRecord<Partial, Layout>(Entry.LogicalAddress, Entry.Data, Entry.Size, Entry.Flags))
+			{
+				Entry = Entry.Advance(Appended);
+			}
+
+			// Append to the current tail block with new records until we've
+			// either filled the block, or run out of data to append.
+			while (const size_t Appended = TailBlock->AppendWithNewRecord<Partial, Layout>(Entry.LogicalAddress, Entry.Data, Entry.Size, Entry.Flags))
+			{
+				NumEntries++;
+				Entry = Entry.Advance(Appended);
+			}
+
+			if (AUTORTFM_LIKELY(Entry.Size == 0))
+			{
+				return;
+			}
+
+			// TailBlock is full, and there's more data to append.
+			// We need to allocate a new block.
+
+			// The maximum number of bytes that can be stored in a single record.
+			const size_t RecordMaxSize = Layout::MaxSize;
+			// Calculate how many records we need to make.
+			const size_t NumRecords = (Entry.Size + RecordMaxSize - 1) / RecordMaxSize;
+			// Calculate the exact required size of the block.
+			const size_t RequiredSize =
+				sizeof(FBlock) +               // FBlock header
+				NumRecords * sizeof(FRecord) + // Bytes needed for the records
+				Entry.Size;                    // Bytes needed for the data
+
+			// Create a new empty tail block, large enough to hold the remainder
+			// of this push in its entirety, and never smaller than the upcoming
+			// block size.
+			AllocateNewBlock(std::max(RequiredSize, NextBlockSize));
+
+			// Append (NumRecords-1) full-records.
+			for (size_t I = 1; I < NumRecords; I++)
+			{
+				TailBlock->AppendWithNewRecord<Whole, Layout>(Entry.LogicalAddress, Entry.Data, RecordMaxSize, Entry.Flags);
+				NumEntries++;
+				Entry = Entry.Advance(RecordMaxSize);
+			}
+
+			// Append last partial record.
+			if (Entry.Size)
+			{
+				TailBlock->AppendWithNewRecord<Whole, Layout>(Entry.LogicalAddress, Entry.Data, Entry.Size, Entry.Flags);
+				NumEntries++;
+			}
+		}
+
+		template <size_t Size, typename Layout>
+		void PushSmall(std::byte* LogicalAddress, std::byte* Data, EWriteFlags Flags)
+		{
+			static_assert(Size <= Layout::MaxSize);
+
+			TotalSizeBytes += Size;
+
+			if (TailBlock->AppendToLastRecord<Whole, Layout>(LogicalAddress, Data, Size, Flags))
+			{
+				return;
+			}
+			
+			if (TailBlock->SpaceRemaining() < sizeof(FRecord) + Size)
+			{
+				AllocateNewBlock(NextBlockSize);
+			}
+			
+			TailBlock->AppendWithNewRecord<Whole, Layout>(LogicalAddress, Data, Size, Flags);
+			NumEntries++;
+		}
+
 		UE_AUTORTFM_FORCEINLINE void AllocateNewBlock(size_t Size)
 		{
-			FBlock* NewBlock = FBlock::Allocate(Size);
+			FBlock* NewBlock = FBlock::Allocate(AlignUp(Size, sizeof(FRecord)));
 			NewBlock->PrevBlock = TailBlock;
 			TailBlock->NextBlock = NewBlock;
 			TailBlock = NewBlock;
 
-			// Increase block sizes by 50% each time, capped at BlockMaxSize.
-			NextBlockSize = std::min<size_t>((NextBlockSize * 3 / 2), BlockMaxSize);
+			// Increase block sizes by 50% each time, capped at MaxBlockSize.
+			NextBlockSize = std::min<size_t>((NextBlockSize * 3 / 2), MaxBlockSize);
 		}
 
 		template <size_t SIZE>
@@ -537,14 +669,13 @@ namespace AutoRTFM
 		// inside the write log.
 		static constexpr size_t HeadBlockSize = 256;
 
-		// The upper bound on heap-allocated block size. We avoid making blocks larger 
-		// than the maximum record size, so we don't need to insert size overflow 
-		// checks throughout the push logic.
-		static constexpr size_t BlockMaxSize = AlignDown(sizeof(FBlock) + sizeof(FRecord) + RecordMaxSize, alignof(FRecord));
+		// The upper bound on heap-allocated block size.
+		static constexpr size_t MaxBlockDataSize = 1024 * 1024 * 16; // Without header, record or alignment.
+		static constexpr size_t MaxBlockSize = AlignUp(sizeof(FBlock) + sizeof(FRecord) + MaxBlockDataSize, alignof(FRecord));
 
 		static_assert(IsAlignedForFRecord<HeadBlockSize>);
 		static_assert(IsAlignedForFRecord<FBlock::DefaultSize>);
-		static_assert(IsAlignedForFRecord<BlockMaxSize>);
+		static_assert(IsAlignedForFRecord<MaxBlockSize>);
 
 		FHash HashAVX2(size_t NumWriteEntries) const;
 
diff --git a/Engine/Source/Runtime/AutoRTFM/Public/AutoRTFM.h b/Engine/Source/Runtime/AutoRTFM/Public/AutoRTFM.h
--- a/Engine/Source/Runtime/AutoRTFM/Public/AutoRTFM.h
+++ b/Engine/Source/Runtime/AutoRTFM/Public/AutoRTFM.h
@@ -28,6 +28,32 @@
 namespace AutoRTFM
 {
 
+enum class EWriteFlags
+{
+	// No flags applied to the write.
+	Default = autortfm_write_default,
+	// The write should not be considered by the memory validator.
+	NoValidation = autortfm_write_no_validation,
+	// The autortfm_extern_api::RollbackWrite will be called to perform
+	// the rollback for this write if the transaction is aborted.
+	CustomRollback = autortfm_write_custom_rollback,
+	// Custom flag bits which can be used by the custom abort callbacks.
+	// Requires CustomRollback for usage.
+	CustomFlag0 = autortfm_write_custom_flag_0,
+	CustomFlag1 = autortfm_write_custom_flag_1,
+};
+
+// Bitwise operator | for EWriteFlags
+inline constexpr EWriteFlags operator | (EWriteFlags LHS, EWriteFlags RHS)
+{
+	return static_cast<EWriteFlags>(static_cast<int>(LHS) | static_cast<int>(RHS));
+}
+// Bitwise operator & for EWriteFlags
+inline constexpr EWriteFlags operator & (EWriteFlags LHS, EWriteFlags RHS)
+{
+	return static_cast<EWriteFlags>(static_cast<int>(LHS) & static_cast<int>(RHS));
+}
+
 // Evaluates to the AutoRTFM mode of the function overload which would be called
 // when calling FuncType with the arguments of type ArgTypes.
 // Warning: This is an experimental API and may be removed in the future.
@@ -671,9 +697,9 @@ UE_AUTORTFM_CRITICAL_INLINE void DidFree(void* Ptr)
 // a possible abort. Normally, tracking memory overwrites should be automatically
 // handled by AutoRTFM, but manual overwrite tracking may be required for third-party
 // libraries or outside compilers (such as ISPC).
-UE_AUTORTFM_CRITICAL_INLINE void RecordOpenWrite(void* Ptr, size_t Size)
+UE_AUTORTFM_CRITICAL_INLINE void RecordOpenWrite(void* Ptr, size_t Size, EWriteFlags Flags = EWriteFlags::Default)
 {
-	autortfm_record_open_write(Ptr, Size);
+	autortfm_record_open_write_with_flags(Ptr, Size, static_cast<autortfm_write_flags>(Flags));
 }
 
 // Informs the runtime that a block of memory is about to be overwritten in the open.
@@ -681,21 +707,21 @@ UE_AUTORTFM_CRITICAL_INLINE void RecordOpenWrite(void* Ptr, size_t Size)
 // a possible abort. Normally, tracking memory overwrites should be automatically
 // handled by AutoRTFM, but manual overwrite tracking may be required for third-party
 // libraries or outside compilers (such as ISPC).
-template<typename Type> UE_AUTORTFM_CRITICAL_INLINE void RecordOpenWrite(Type* Ptr)
+template<typename Type> UE_AUTORTFM_CRITICAL_INLINE void RecordOpenWrite(Type* Ptr, EWriteFlags Flags = EWriteFlags::Default)
 {
-	autortfm_record_open_write(Ptr, sizeof(Type));
+	autortfm_record_open_write_with_flags(Ptr, sizeof(Type), static_cast<autortfm_write_flags>(Flags));
 }
 
 // Same as RecordOpenWrite() but marks the write as ignorable by the memory validator.
 UE_AUTORTFM_CRITICAL_INLINE void RecordOpenWriteNoMemoryValidation(void* Ptr, size_t Size)
 {
-	autortfm_record_open_write_no_memory_validation(Ptr, Size);
+	autortfm_record_open_write_with_flags(Ptr, Size, autortfm_write_no_validation);
 }
 
 // Same as RecordOpenWrite() but marks the write as ignorable by the memory validator.
 template<typename Type> UE_AUTORTFM_CRITICAL_INLINE void RecordOpenWriteNoMemoryValidation(Type* Ptr)
 {
-	autortfm_record_open_write_no_memory_validation(Ptr, sizeof(Type));
+	autortfm_record_open_write_with_flags(Ptr, sizeof(Type), autortfm_write_no_validation);
 }
 
 // Report that a unreachable codepath is being hit. Used to manually ban certain codepaths
diff --git a/Engine/Source/Runtime/AutoRTFM/Public/AutoRTFM/CAPI.h b/Engine/Source/Runtime/AutoRTFM/Public/AutoRTFM/CAPI.h
--- a/Engine/Source/Runtime/AutoRTFM/Public/AutoRTFM/CAPI.h
+++ b/Engine/Source/Runtime/AutoRTFM/Public/AutoRTFM/CAPI.h
@@ -75,6 +75,22 @@ typedef enum
 	autortfm_log_fatal,
 } autortfm_log_severity;
 
+// Flags that can be passed to autortfm_record_open_write_with_flags().
+typedef enum
+{
+	// No flags applied to the write.
+	autortfm_write_default = 0x0,
+	// The write should not be considered by the memory validator.
+	autortfm_write_no_validation = 0x1,
+	// The autortfm_extern_api::RollbackWrite will be called to perform
+	// the rollback for this write if the transaction is aborted.
+	autortfm_write_custom_rollback = 0x2,
+	// Custom flag bits which can be used by the custom abort callbacks.
+	// Requires autortfm_write_custom_rollback for usage.
+	autortfm_write_custom_flag_0 = 0x4,
+	autortfm_write_custom_flag_1 = 0x8,
+} autortfm_write_flags;
+
 // An opaque unique identifier for a transaction.
 typedef uint64_t autortfm_transaction_id;
 
@@ -146,6 +162,20 @@ typedef struct
 	// Can be null.
 	void (*OnMemoryValidationStatisticsChanged)();
 
+	// Custom abort handler for writes that that were added with the
+	// autortfm_write_custom_rollback flag. It is the responsibility of this
+	// function to perform the rollback (copy of Size bytes from Value to
+	// Address).
+	// RollbackWrite may contain multiple writes condensed into a single
+	// contigious block of memory, and the writes may be split into multiple
+	// calls with no defined size or alignment granularity.
+	// Calls are ordered starting with the most recently recorded writes.
+	// Address / Value pointers are ordered starting with the least recently
+	// recorded writes.
+	// Split writes are guaranteed to be "completed" with the next immediate
+	// call(s).
+	void (*RollbackWrite)(void* Address, const void* Value, size_t Size, autortfm_write_flags Flags);
+
 } autortfm_extern_api;
 
 #if UE_AUTORTFM_ENABLED
@@ -287,12 +317,13 @@ UE_AUTORTFM_CRITICAL_INLINE void autortfm_record_open_write(void* Ptr, size_t Si
 #endif
 
 #if UE_AUTORTFM_ENABLED
-UE_AUTORTFM_API void autortfm_record_open_write_no_memory_validation(void* Ptr, size_t Size) AUTORTFM_NOEXCEPT;
+UE_AUTORTFM_API void autortfm_record_open_write_with_flags(void* Ptr, size_t Size, autortfm_write_flags Flags) AUTORTFM_NOEXCEPT;
 #else
-UE_AUTORTFM_CRITICAL_INLINE void autortfm_record_open_write_no_memory_validation(void* Ptr, size_t Size) AUTORTFM_NOEXCEPT
+UE_AUTORTFM_CRITICAL_INLINE void autortfm_record_open_write_with_flags(void* Ptr, size_t Size, autortfm_write_flags Flags) AUTORTFM_NOEXCEPT
 {
 	UE_AUTORTFM_UNUSED(Ptr);
 	UE_AUTORTFM_UNUSED(Size);
+	UE_AUTORTFM_UNUSED(Flags);
 }
 #endif
 
