[AutoRTFM] Add all HitSet entries to the parent transaction on commit, not just the coalesced write logs

For back-to-back writes, this will eliminate unnecessary write log copying from child transactions into the parent.

Also add missing unit tests for FHitSet.

#Synced-CL 48740692

diff --git a/Engine/Source/Programs/AutoRTFMTests/Private/HitSetTests.cpp b/Engine/Source/Programs/AutoRTFMTests/Private/HitSetTests.cpp
new file mode 100644
--- /dev/null
+++ b/Engine/Source/Programs/AutoRTFMTests/Private/HitSetTests.cpp
@@ -0,0 +1,194 @@
+// Copyright Epic Games, Inc. All Rights Reserved.
+
+#include "HitSet.h"
+#include "Catch2Includes.h"
+#include <cstdint>
+#include <unordered_set>
+
+TEST_CASE("HitSet")
+{
+	struct FHasher
+	{
+		size_t operator()(AutoRTFM::FHitSetEntry Entry) const
+		{
+			return static_cast<size_t>(Entry.Payload());
+		}
+	};
+
+	struct FEqualTo
+	{
+		bool operator()(AutoRTFM::FHitSetEntry A, AutoRTFM::FHitSetEntry B) const
+		{
+			return A.Payload() == B.Payload();
+		}
+	};
+
+	AutoRTFM::FHitSet HitSet;
+	std::unordered_set<AutoRTFM::FHitSetEntry, FHasher, FEqualTo> Expected;
+
+	auto CheckIsExpected = [&]
+	{
+		REQUIRE(HitSet.GetCount() == Expected.size());
+		if (Expected.size() == 0)
+		{
+			REQUIRE(HitSet.IsEmpty());
+			REQUIRE(HitSet.begin() == HitSet.end());
+		}
+		else
+		{
+			REQUIRE(!HitSet.IsEmpty());
+			REQUIRE(HitSet.begin() != HitSet.end());
+		}
+
+		REQUIRE(HitSet.GetCount() <= HitSet.GetCapacity());
+
+		std::unordered_set<AutoRTFM::FHitSetEntry, FHasher, FEqualTo> Got;
+		for (AutoRTFM::FHitSetEntry GotEntry : HitSet)
+		{
+			Got.insert(GotEntry);
+			const bool bFoundHitSetEntryInExpected = Expected.contains(GotEntry);
+			REQUIRE(bFoundHitSetEntryInExpected);
+		}
+		for (AutoRTFM::FHitSetEntry ExpectedEntry : Expected)
+		{
+			const bool bFoundExpectedInHitSet = Got.contains(ExpectedEntry);
+			REQUIRE(bFoundExpectedInHitSet);
+			REQUIRE(HitSet.Contains(ExpectedEntry));
+		}
+		REQUIRE(Got.size() == Expected.size());
+	};
+
+	auto RunTests = [&](auto Insert)
+	{
+		constexpr uintptr_t One = 1;
+		std::mt19937 Rand(0x1234);
+
+		SECTION("Vary Address")
+		{
+			const uintptr_t Size = 0x100;
+			const uintptr_t Flags = 0;
+			for (size_t I = 1; I < 1000; I++)
+			{
+				const uintptr_t Address = I * 0x1000;
+				Insert({Address, Size, Flags});
+			}
+			CheckIsExpected();
+			for (size_t I = 1; I < 2000; I++)
+			{
+				const uintptr_t Address = I * 0x500;
+				Insert({Address, Size, Flags});
+			}
+			CheckIsExpected();
+		}
+
+		SECTION("Vary Size")
+		{
+			const uintptr_t Address = 0x1000;
+			const uintptr_t Flags = 0;
+			for (size_t I = 0; I < 1000; I++)
+			{
+				const uintptr_t Size = I;
+				Insert({Address, I, Flags});
+			}
+			CheckIsExpected();
+			for (size_t I = 0; I < 1000; I++)
+			{
+				const uintptr_t Size = I;
+				Insert({Address, I, Flags});
+			}
+			CheckIsExpected();
+		}
+
+		SECTION("Vary Flags")
+		{
+			const uintptr_t Address = 0x1000;
+			const uintptr_t Size = 0x100;
+			for (size_t I = 0; I < 15; I++)
+			{
+				const uintptr_t Flags = I;
+				Insert({Address, Size, Flags});
+			}
+			CheckIsExpected();
+			for (size_t I = 0; I < 15; I++)
+			{
+				const uintptr_t Flags = I;
+				Insert({Address, Size, Flags});
+			}
+			CheckIsExpected();
+		}
+
+		SECTION("Vary All")
+		{
+			for (size_t I = 0; I < 10000; I++)
+			{
+				const uintptr_t Address = One << (Rand() % AutoRTFM::FHitSetEntry::AddressBits);
+				const uintptr_t Size = One << (Rand() % AutoRTFM::FHitSetEntry::SizeBits);
+				const uintptr_t Flags = Rand() & ((One << AutoRTFM::FHitSetEntry::FlagsBits) - 1);
+				Insert({Address, Size, Flags});
+			}
+			CheckIsExpected();
+		}
+
+		{
+			for (size_t I = 0; I < 10000; I++)
+			{
+				const uintptr_t Address = One << (Rand() % AutoRTFM::FHitSetEntry::AddressBits);
+				const uintptr_t Size = One << (Rand() % AutoRTFM::FHitSetEntry::SizeBits);
+				const uintptr_t Flags = Rand() & ((One << AutoRTFM::FHitSetEntry::FlagsBits) - 1);
+				const AutoRTFM::FHitSetEntry Entry{Address, Size, Flags};
+				REQUIRE(HitSet.Contains(Entry) == Expected.contains(Entry));
+			}
+		}
+	};
+
+	SECTION("FindOrTryInsert")
+	{
+		auto Insert = [&](AutoRTFM::FHitSetEntry Entry)
+		{
+			// We're not testing for max-capacity behavior here.
+			REQUIRE(HitSet.GetCapacity() < AutoRTFM::FHitSet::MaxCapacity);
+
+			const bool bAlreadyAdded = Expected.contains(Entry);
+			
+			if (HitSet.FindOrTryInsert(Entry))
+			{
+				REQUIRE(bAlreadyAdded);
+			}
+			else
+			{
+				REQUIRE(!bAlreadyAdded);
+				Expected.insert(Entry);
+			}
+		};
+
+		RunTests(Insert);
+	}
+
+	SECTION("FindOrTryInsertNoResize")
+	{
+		auto InsertNoResize = [&](AutoRTFM::FHitSetEntry Entry)
+		{
+			switch (HitSet.FindOrTryInsertNoResize(Entry))
+			{
+				case AutoRTFM::FHitSet::EInsertResult::Inserted:
+				{
+					const bool bAdded = Expected.emplace(Entry).second;
+					REQUIRE(bAdded);
+					break;
+				}
+				case AutoRTFM::FHitSet::EInsertResult::NotInserted:
+				{
+					REQUIRE(!Expected.contains(Entry));
+					break;
+				}
+				case AutoRTFM::FHitSet::EInsertResult::Exists:
+				{
+					REQUIRE(Expected.contains(Entry));
+					break;
+				}
+			}
+		};
+
+		RunTests(InsertNoResize);
+	}
+}
diff --git a/Engine/Source/Runtime/AutoRTFM/Private/HitSet.h b/Engine/Source/Runtime/AutoRTFM/Private/HitSet.h
--- a/Engine/Source/Runtime/AutoRTFM/Private/HitSet.h
+++ b/Engine/Source/Runtime/AutoRTFM/Private/HitSet.h
@@ -13,13 +13,17 @@ namespace AutoRTFM
 {
 	struct FHitSetEntry
 	{
+		static constexpr size_t AddressBits = 48;
+		static constexpr size_t SizeBits = 12;
+		static constexpr size_t FlagsBits = 4;
+
 		// The address of the write.
-		uintptr_t Address : 48;
+		uintptr_t Address : AddressBits;
 		// Size of the write.
-		uintptr_t Size : 12;
+		uintptr_t Size : SizeBits;
 		// Flags applied to the write.
-		uintptr_t Flags : 4;
-		
+		uintptr_t Flags : FlagsBits;
+
 		UE_AUTORTFM_FORCEINLINE uintptr_t Payload() const
 		{
 			// Note: Clang optimizes this away to a reinterpret cast.
@@ -42,14 +46,13 @@ namespace AutoRTFM
 		// recording meaningless hits, against the potential to hit again.
 		static constexpr size_t MaxSize = 16;
 
+		// TODO: Revisit a good max capacity for the hitset.
+		static constexpr uint64_t MaxCapacity = static_cast<uint64_t>(128 * 1024 * 1024) / sizeof(FHitSetEntry);
 	private:
 		static constexpr uint8_t LinearProbeDepth = 32;
 		static constexpr uint8_t LogInitialCapacity = 6;
 		static constexpr uint64_t InitialCapacity = static_cast<uint64_t>(1) << LogInitialCapacity;
 
-		// TODO: Revisit a good max capacity for the hitset.
-		static constexpr uint64_t MaxCapacity = static_cast<uint64_t>(128 * 1024 * 1024) / sizeof(FHitSetEntry);
-
 		static_assert(LinearProbeDepth <= InitialCapacity);
 	public:
 
@@ -82,7 +85,7 @@ namespace AutoRTFM
 				switch (InsertNoResize(Entry.Payload()))
 				{
 				default:
-					AutoRTFM::InternalUnreachable();
+					InternalUnreachable();
 				case EInsertResult::Inserted:
 					return false;
 				case EInsertResult::Exists:
@@ -97,6 +100,55 @@ namespace AutoRTFM
 			}
 		}
 
+		// Returns true if the HitSet contains the given entry, without mutation.
+		UE_AUTORTFM_FORCEINLINE bool Contains(FHitSetEntry Entry) const
+		{
+			uintptr_t Raw = Entry.Payload();
+			const uintptr_t Hash = FibonacciHash(Raw);
+
+			// The first index will always be in the range of capacity because we
+			// are using a fibonacci hash.
+			{
+				const uintptr_t Value = Payload[Hash];
+				if (0 == Value)
+				{
+					return false;
+				}
+				if (Raw == Value)
+				{
+					return true;
+				}
+				// Otherwise we need to do a linear probe...
+			}
+
+			const uintptr_t Mask = Capacity() - 1;
+
+			// Clang goes way over the top with loop unrolling, and makes this code noticeably
+			// slower as a result. So we disable it!
+#ifdef __clang__
+#pragma clang loop unroll(disable)
+#endif
+			for (uint8_t D = 1; D < LinearProbeDepth; D++)
+			{
+				// Capacity is always a power of 2, so we can just mask out the bits.
+				const uintptr_t I = (Hash + D) & Mask;
+
+				const uintptr_t Value = Payload[I];
+				if (0 == Value)
+				{
+					return false;
+				}
+				if (Raw == Value)
+				{
+					return true;
+				}
+
+				// Otherwise we need to keep going with our linear probe...
+			}
+
+			return false;
+		}
+
 		UE_AUTORTFM_FORCEINLINE bool IsEmpty() const
 		{
 			return 0 == Count;
@@ -128,6 +180,59 @@ namespace AutoRTFM
 			return Count;
 		}
 
+		struct FIterator
+		{
+			FIterator(uintptr_t const * Ptr, size_t EntriesRemaining)
+				: Ptr{Ptr}, EntriesRemaining{EntriesRemaining}
+			{
+				SkipEmpty();
+			}
+
+			FHitSetEntry operator*() const
+			{
+				FHitSetEntry Out;
+				memcpy(&Out, Ptr, sizeof(Out));
+				return Out;
+			}
+
+			FIterator& operator++()
+			{
+				if (EntriesRemaining)
+				{
+					EntriesRemaining--;
+					Ptr++;
+					SkipEmpty();
+				}
+				return *this;
+			}
+
+			bool operator==(FIterator Other) const
+			{
+				return Other.Ptr == Ptr;
+			}
+
+			bool operator!=(FIterator Other) const
+			{
+				return Other.Ptr != Ptr;
+			}
+
+		private:
+			void SkipEmpty()
+			{
+				while (EntriesRemaining > 0 && *Ptr == 0)
+				{
+					EntriesRemaining--;
+					Ptr++;
+				}
+			}
+
+			const uintptr_t* Ptr;
+			size_t EntriesRemaining;
+		};
+		
+		FIterator begin() const { return FIterator(Payload, Capacity()); }
+		FIterator end() const { return FIterator(Payload + Capacity(), 0); }
+
 	private:
 		uintptr_t SmallPayload[InitialCapacity]{};
 		uintptr_t* Payload = SmallPayload;
@@ -246,6 +351,8 @@ namespace AutoRTFM
 		// succeeded (EG. the key was not already in the set).
 		UE_AUTORTFM_FORCEINLINE EInsertResult InsertNoResize(const uintptr_t Raw)
 		{
+			AUTORTFM_ASSERT_DEBUG(Raw != 0);
+
 			const uintptr_t Hash = FibonacciHash(Raw);
 
 			// The first index will always be in the range of capacity because we
@@ -265,7 +372,7 @@ namespace AutoRTFM
 
 			const uintptr_t Mask = Capacity() - 1;
 
-			// Clang goes way over the top with loop unrolling, and makes this code noticably
+			// Clang goes way over the top with loop unrolling, and makes this code noticeably
 			// slower as a result. So we disable it!
 #ifdef __clang__
 #pragma clang loop unroll(disable)
diff --git a/Engine/Source/Runtime/AutoRTFM/Private/Transaction.cpp b/Engine/Source/Runtime/AutoRTFM/Private/Transaction.cpp
--- a/Engine/Source/Runtime/AutoRTFM/Private/Transaction.cpp
+++ b/Engine/Source/Runtime/AutoRTFM/Private/Transaction.cpp
@@ -375,6 +375,16 @@ void FTransaction::CommitNested()
 		Parent->WriteLog.Push(Write);
 	}
 
+	// Similarly we need to merge the hit-sets to the parent, skipping those
+	// that are on the parent transaction's stack.
+	for (FHitSetEntry Entry : HitSet)
+	{
+		if (!Parent->IsOnStack(reinterpret_cast<void*>(Entry.Address)))
+		{
+			Parent->HitSet.FindOrTryInsert(Entry);
+		}
+	}
+
 	// For all the deferred calls to `PopOnCommitHandler` that we couldn't
 	// process (because our transaction nest didn't `PushOnCommitHandler`)
 	// we need to move these to the parent now to handle them.
