[AutoRTFM] Simplify ClosedInstrumentationVisitor, pull sharable bits out to Utils.[h|cpp]

The sharable parts will be used by the new AutoRTFM sanitizer.

#Synced-CL 48740692

diff --git a/Engine/Source/Runtime/AutoRTFM/Private/CompilerInterface.cpp b/Engine/Source/Runtime/AutoRTFM/Private/CompilerInterface.cpp
--- a/Engine/Source/Runtime/AutoRTFM/Private/CompilerInterface.cpp
+++ b/Engine/Source/Runtime/AutoRTFM/Private/CompilerInterface.cpp
@@ -71,12 +71,12 @@ extern "C" AUTORTFM_DISABLE UE_AUTORTFM_API AUTORTFM_NO_ASAN void autortfm_recor
 	Context->RecordWrite<8>(Ptr);
 }
 
-extern "C" AUTORTFM_DISABLE UE_AUTORTFM_API AUTORTFM_NO_ASAN void autortfm_record_masked_write(void* Ptr, uintptr_t Mask, int MaskWidthBits, int ValueSizeBytes) noexcept
+extern "C" AUTORTFM_DISABLE UE_AUTORTFM_API AUTORTFM_NO_ASAN void autortfm_record_masked_write(void* Ptr, uintptr_t Mask, int NumMaskBits, int ValueSizeBytes) noexcept
 {
 	UE_AUTORTFM_HANDLE_NULL_WRITE(Ptr);
 
 	char* IncrementablePtr = static_cast<char*>(Ptr);
-	for(int i = 0; i < MaskWidthBits; i++)
+	for(int i = 0; i < NumMaskBits; i++)
 	{
 		if (Mask & (1u << i))
 		{
