[AutoRTFM] Update autortfm_mode_disable_no_sanitize comment based on post-submit review comments

#Synced-CL 48740692

diff --git a/Engine/Source/Runtime/AutoRTFM/Public/AutoRTFM/Constants.h b/Engine/Source/Runtime/AutoRTFM/Public/AutoRTFM/Constants.h
--- a/Engine/Source/Runtime/AutoRTFM/Public/AutoRTFM/Constants.h
+++ b/Engine/Source/Runtime/AutoRTFM/Public/AutoRTFM/Constants.h
@@ -46,9 +46,9 @@ typedef enum
 	autortfm_mode_disable = 1,
 
 	// [[clang::autortfm(disable_no_sanitize)]]
-	// Same as autortfm_mode_disable, but the function will ignored by the
-	// AutoRTFM sanitizer. Note that calls made by this function will still
-	// be instrumented by the AutoRTFM sanitizer.
+	// Same as autortfm_mode_disable, but the function will be ignored by the
+	// AutoRTFM sanitizer. Note that the sanitizer will still be active for
+	// function calls made by this function.
 	autortfm_mode_disable_no_sanitize = 2,
 
 	// [[clang::autortfm(infer)]]
