[AutoRTFM] Add runtime support for AutoRTFMSan, a new AutoRTFM sanitizer for detecting mixed closed / open writes.

This serves the same goals as the memory validator, but can:
* Identify the specific open write that collides with a previous closed write, without reliance on aborts / cascading retries which upsets the VVM.
* Be performant enough to not require throttling.
* Avoid any runtime overhead when not enabled.

#Synced-CL 48740692

diff --git a/Engine/Source/Programs/AutoRTFMTests/Private/StackTests.cpp b/Engine/Source/Programs/AutoRTFMTests/Private/StackTests.cpp
--- a/Engine/Source/Programs/AutoRTFMTests/Private/StackTests.cpp
+++ b/Engine/Source/Programs/AutoRTFMTests/Private/StackTests.cpp
@@ -1,5 +1,6 @@
 // Copyright Epic Games, Inc. All Rights Reserved.
 
+#include "AutoRTFM.h"
 #include "Stack.h"
 #include "Catch2Includes.h"
 #include "ObjectLifetimeHelper.h"
@@ -9,13 +10,15 @@ using NonTrivialStack = AutoRTFM::TStack<AutoRTFMTestUtils::FObjectLifetimeHelpe
 
 TEMPLATE_TEST_CASE("Stack", "", TrivialStack, NonTrivialStack)
 {
+	AutoRTFM::UnreachableIfClosed();
+
 	REQUIRE(AutoRTFMTestUtils::FObjectLifetimeHelper::ConstructorCalls == 0);
 	REQUIRE(AutoRTFMTestUtils::FObjectLifetimeHelper::DestructorCalls == 0);
 
 	using StackType = TestType;
 	using ElementType = typename TestType::ElementType;
 
-	auto Check = [&](StackType& Stack, std::vector<ElementType> Expected)
+	auto Check = [&] AUTORTFM_DISABLE (StackType& Stack, std::vector<ElementType> Expected)
 	{
 		// Check the reported count is as expected.
 		REQUIRE(Stack.Num() == Expected.size());
@@ -225,7 +228,7 @@ TEMPLATE_TEST_CASE("Stack", "", TrivialStack, NonTrivialStack)
 	{
 		StackType Source;
 		StackType Target;
-		auto CheckCopy = [&]()
+		auto CheckCopy = [&] AUTORTFM_DISABLE
 		{
 			SECTION("Source Empty")
 			{
@@ -310,7 +313,7 @@ TEMPLATE_TEST_CASE("Stack", "", TrivialStack, NonTrivialStack)
 	{
 		StackType Source;
 		StackType Target;
-		auto CheckMove = [&]
+		auto CheckMove = [&] AUTORTFM_DISABLE
 		{
 			SECTION("Source Empty")
 			{
diff --git a/Engine/Source/Runtime/AutoRTFM/Private/API.cpp b/Engine/Source/Runtime/AutoRTFM/Private/API.cpp
--- a/Engine/Source/Runtime/AutoRTFM/Private/API.cpp
+++ b/Engine/Source/Runtime/AutoRTFM/Private/API.cpp
@@ -15,7 +15,11 @@ static_assert(UE_AUTORTFM_ENABLED, "AutoRTFM/API.cpp requires the compiler flag
 
 namespace
 {
+#if AUTORTFM_SANITIZER
+	static constexpr AutoRTFM::EMemoryValidationLevel GDefaultMemoryValidationLevel = AutoRTFM::EMemoryValidationLevel::Error;
+#else
 	static constexpr AutoRTFM::EMemoryValidationLevel GDefaultMemoryValidationLevel = AutoRTFM::EMemoryValidationLevel::Disabled;
+#endif
 
 	// Move this to a local only and use functions to access this
 #if UE_AUTORTFM_ENABLED_RUNTIME_BY_DEFAULT
diff --git a/Engine/Source/Runtime/AutoRTFM/Private/CompilerInterface.cpp b/Engine/Source/Runtime/AutoRTFM/Private/CompilerInterface.cpp
--- a/Engine/Source/Runtime/AutoRTFM/Private/CompilerInterface.cpp
+++ b/Engine/Source/Runtime/AutoRTFM/Private/CompilerInterface.cpp
@@ -1,6 +1,5 @@
 // Copyright Epic Games, Inc. All Rights Reserved.
 
-#include "AutoRTFM/Defines.h"
 #if (defined(__AUTORTFM) && __AUTORTFM)
 
 #include "AutoRTFM.h"
@@ -9,6 +8,7 @@
 #include "ExternAPI.h"
 #include "FunctionMapInlines.h"
 #include "Memcpy.h"
+#include "ScopedGuard.h"
 
 namespace AutoRTFM
 {
@@ -36,42 +36,42 @@ namespace AutoRTFM
 
 #endif
 
-extern "C" AUTORTFM_DISABLE UE_AUTORTFM_API AUTORTFM_NO_ASAN void autortfm_record_write(void* Ptr, size_t Size) noexcept
+extern "C" AUTORTFM_DISABLE_NO_SANITIZE UE_AUTORTFM_API AUTORTFM_NO_ASAN void autortfm_record_write(void* Ptr, size_t Size) noexcept
 {
 	UE_AUTORTFM_HANDLE_NULL_WRITE(Ptr);
 	FContext* Context = FContext::Get();
 	Context->RecordWrite(Ptr, Size);
 }
 
-extern "C" AUTORTFM_DISABLE UE_AUTORTFM_API AUTORTFM_NO_ASAN void autortfm_record_write_1(void* Ptr) noexcept
+extern "C" AUTORTFM_DISABLE_NO_SANITIZE UE_AUTORTFM_API AUTORTFM_NO_ASAN void autortfm_record_write_1(void* Ptr) noexcept
 {
 	UE_AUTORTFM_HANDLE_NULL_WRITE(Ptr);
 	FContext* Context = FContext::Get();
 	Context->RecordWrite<1>(Ptr);
 }
 
-extern "C" AUTORTFM_DISABLE UE_AUTORTFM_API AUTORTFM_NO_ASAN void autortfm_record_write_2(void* Ptr) noexcept
+extern "C" AUTORTFM_DISABLE_NO_SANITIZE UE_AUTORTFM_API AUTORTFM_NO_ASAN void autortfm_record_write_2(void* Ptr) noexcept
 {
 	UE_AUTORTFM_HANDLE_NULL_WRITE(Ptr);
 	FContext* Context = FContext::Get();
 	Context->RecordWrite<2>(Ptr);
 }
 
-extern "C" AUTORTFM_DISABLE UE_AUTORTFM_API AUTORTFM_NO_ASAN void autortfm_record_write_4(void* Ptr) noexcept
+extern "C" AUTORTFM_DISABLE_NO_SANITIZE UE_AUTORTFM_API AUTORTFM_NO_ASAN void autortfm_record_write_4(void* Ptr) noexcept
 {
 	UE_AUTORTFM_HANDLE_NULL_WRITE(Ptr);
 	FContext* Context = FContext::Get();
 	Context->RecordWrite<4>(Ptr);
 }
 
-extern "C" AUTORTFM_DISABLE UE_AUTORTFM_API AUTORTFM_NO_ASAN void autortfm_record_write_8(void* Ptr) noexcept
+extern "C" AUTORTFM_DISABLE_NO_SANITIZE UE_AUTORTFM_API AUTORTFM_NO_ASAN void autortfm_record_write_8(void* Ptr) noexcept
 {
 	UE_AUTORTFM_HANDLE_NULL_WRITE(Ptr);
 	FContext* Context = FContext::Get();
 	Context->RecordWrite<8>(Ptr);
 }
 
-extern "C" AUTORTFM_DISABLE UE_AUTORTFM_API AUTORTFM_NO_ASAN void autortfm_record_masked_write(void* Ptr, uintptr_t Mask, int NumMaskBits, int ValueSizeBytes) noexcept
+extern "C" AUTORTFM_DISABLE_NO_SANITIZE UE_AUTORTFM_API AUTORTFM_NO_ASAN void autortfm_record_masked_write(void* Ptr, uintptr_t Mask, int NumMaskBits, int ValueSizeBytes) noexcept
 {
 	UE_AUTORTFM_HANDLE_NULL_WRITE(Ptr);
 
@@ -151,19 +151,19 @@ extern "C" void* autortfm_lookup_function(void* OriginalFunction, const char* Wh
 	return FunctionMapLookup(OriginalFunction, Where);
 }
 
-extern "C" AUTORTFM_DISABLE UE_AUTORTFM_API void autortfm_memcpy(void* Dst, const void* Src, size_t Size) noexcept
+extern "C" AUTORTFM_DISABLE_NO_SANITIZE UE_AUTORTFM_API void autortfm_memcpy(void* Dst, const void* Src, size_t Size) noexcept
 {
 	FContext* Context = FContext::Get();
     Memcpy(Dst, Src, Size, Context);
 }
 
-extern "C" AUTORTFM_DISABLE UE_AUTORTFM_API void autortfm_memmove(void* Dst, const void* Src, size_t Size) noexcept
+extern "C" AUTORTFM_DISABLE_NO_SANITIZE UE_AUTORTFM_API void autortfm_memmove(void* Dst, const void* Src, size_t Size) noexcept
 {
 	FContext* Context = FContext::Get();
     Memmove(Dst, Src, Size, Context);
 }
 
-extern "C" AUTORTFM_DISABLE UE_AUTORTFM_API void autortfm_memset(void* Dst, int Value, size_t Size) noexcept
+extern "C" AUTORTFM_DISABLE_NO_SANITIZE UE_AUTORTFM_API void autortfm_memset(void* Dst, int Value, size_t Size) noexcept
 {
 	FContext* Context = FContext::Get();
     Memset(Dst, Value, Size, Context);
@@ -175,7 +175,7 @@ extern "C" UE_AUTORTFM_ALWAYS_OPEN UE_AUTORTFM_API void autortfm_unreachable(con
 	__builtin_unreachable();
 }
 
-extern "C" AUTORTFM_DISABLE UE_AUTORTFM_API void autortfm_llvm_fail(const char* Message) noexcept
+extern "C" AUTORTFM_DISABLE_NO_SANITIZE UE_AUTORTFM_API void autortfm_llvm_fail(const char* Message) noexcept
 {
 	if (Message)
 	{
@@ -187,7 +187,7 @@ extern "C" AUTORTFM_DISABLE UE_AUTORTFM_API void autortfm_llvm_fail(const char*
 	}
 }
 
-extern "C" AUTORTFM_DISABLE UE_AUTORTFM_API void autortfm_llvm_missing_function() noexcept
+extern "C" AUTORTFM_DISABLE_NO_SANITIZE UE_AUTORTFM_API void autortfm_llvm_missing_function() noexcept
 {
 	if (ForTheRuntime::GetInternalAbortAction() == ForTheRuntime::EAutoRTFMInternalAbortActionState::Crash)
 	{
@@ -202,15 +202,34 @@ extern "C" AUTORTFM_DISABLE UE_AUTORTFM_API void autortfm_llvm_missing_function(
     Context->AbortByLanguageAndThrow();
 }
 
-extern "C" AUTORTFM_DISABLE UE_AUTORTFM_FORCENOINLINE UE_AUTORTFM_API void autortfm_called_no_autortfm() noexcept
+extern "C" AUTORTFM_DISABLE_NO_SANITIZE UE_AUTORTFM_FORCENOINLINE UE_AUTORTFM_API void autortfm_called_no_autortfm() noexcept
 {
 	AUTORTFM_FATAL("inlined UE_AUTORTFM_NOAUTORTFM function called from the closed");
 }
 
-extern "C" AUTORTFM_DISABLE UE_AUTORTFM_FORCENOINLINE UE_AUTORTFM_API void autortfm_disable_called() noexcept
+extern "C" AUTORTFM_DISABLE_NO_SANITIZE UE_AUTORTFM_FORCENOINLINE UE_AUTORTFM_API void autortfm_disable_called() noexcept
 {
 	AUTORTFM_FATAL("inlined AUTORTFM_DISABLE function called from the closed");
 }
+
+#if AUTORTFM_SANITIZER
+static bool bInAutoRTFMSanitizerOpenWrite = false;
+extern "C" AUTORTFM_DISABLE_NO_SANITIZE UE_AUTORTFM_FORCENOINLINE UE_AUTORTFM_API void autortfm_sanitizer_open_write(void* Ptr, size_t Size) noexcept
+{
+	if (FContext* Context = FContext::Get(); Context && Context->IsContextThread())
+	{
+		if (!bInAutoRTFMSanitizerOpenWrite)
+		{
+			TScopedGuard RecursionGuard(bInAutoRTFMSanitizerOpenWrite, true);
+			if (FTransaction* Transaction = Context->GetMaterializedTransaction())
+			{
+				Transaction->SanitizerOpenWrite(Ptr, Size);
+			}
+		}
+	}
+}
+#endif
+
 } // namespace AutoRTFM
 
 #endif // defined(__AUTORTFM) && __AUTORTFM
diff --git a/Engine/Source/Runtime/AutoRTFM/Private/Context.h b/Engine/Source/Runtime/AutoRTFM/Private/Context.h
--- a/Engine/Source/Runtime/AutoRTFM/Private/Context.h
+++ b/Engine/Source/Runtime/AutoRTFM/Private/Context.h
@@ -18,7 +18,7 @@
 namespace AutoRTFM
 {
 
-class AUTORTFM_DISABLE FContext final
+class AUTORTFM_DISABLE_NO_SANITIZE FContext final
 {
 	using FTaskPool = typename TTaskArray<TTask<void()>>::FEntryPool;
 	using FTransactionPool = TIntrusivePool<FTransaction, 16>;
@@ -83,7 +83,8 @@ public:
 	}
 
 	inline FCallNest* GetCurrentNest() const { return CurrentNest; }
-	inline EContextStatus GetStatus() const { return CurrentThreadId == FThreadID::GetCurrent() ? Status : EContextStatus::Idle; }
+	inline bool IsContextThread() const { return CurrentThreadId == FThreadID::GetCurrent(); }
+	inline EContextStatus GetStatus() const { return IsContextThread() ? Status : EContextStatus::Idle; }
 	[[noreturn]] void Throw();
 
 	inline FOpenHashThrottler& GetOpenHashThrottler() { return OpenHashThrottler; }
diff --git a/Engine/Source/Runtime/AutoRTFM/Private/FunctionMap.cpp b/Engine/Source/Runtime/AutoRTFM/Private/FunctionMap.cpp
--- a/Engine/Source/Runtime/AutoRTFM/Private/FunctionMap.cpp
+++ b/Engine/Source/Runtime/AutoRTFM/Private/FunctionMap.cpp
@@ -259,7 +259,7 @@ inline void* FollowRelocation(void* Function)
 }
 
 
-AUTORTFM_DISABLE static void* FunctionMapReportError(void* Open, const char* Where)
+AUTORTFM_DISABLE_NO_SANITIZE static void* FunctionMapReportError(void* Open, const char* Where)
 {
 #if AUTORTFM_PLATFORM_WINDOWS
 	if (Open == &InitializeCriticalSection ||
diff --git a/Engine/Source/Runtime/AutoRTFM/Private/FunctionMap.h b/Engine/Source/Runtime/AutoRTFM/Private/FunctionMap.h
--- a/Engine/Source/Runtime/AutoRTFM/Private/FunctionMap.h
+++ b/Engine/Source/Runtime/AutoRTFM/Private/FunctionMap.h
@@ -10,17 +10,17 @@ namespace AutoRTFM
 {
 
 // Dumps statistics about the function map to LOG_INFO.
-AUTORTFM_DISABLE void FunctionMapDumpStats();
+AUTORTFM_DISABLE_NO_SANITIZE void FunctionMapDumpStats();
 
 // Adds all the open to closed function mappings from the linked list of tables to the function map.
-AUTORTFM_DISABLE void FunctionMapAdd(autortfm_open_to_closed_table* Tables);
+AUTORTFM_DISABLE_NO_SANITIZE void FunctionMapAdd(autortfm_open_to_closed_table* Tables);
 
 // Looks up the closed function from the open function using the function map, following dynamic
 // library thunk redirections.
-AUTORTFM_DISABLE void* FunctionMapLookupExhaustive(void* OpenFn, const char* Where);
+AUTORTFM_DISABLE_NO_SANITIZE void* FunctionMapLookupExhaustive(void* OpenFn, const char* Where);
 
 template<typename ReturnType, typename... ParameterTypes>
-AUTORTFM_DISABLE auto FunctionMapLookupExhaustive(ReturnType (*Function)(ParameterTypes...), const char* Where) -> ReturnType (*)(ParameterTypes...)
+AUTORTFM_DISABLE_NO_SANITIZE auto FunctionMapLookupExhaustive(ReturnType (*Function)(ParameterTypes...), const char* Where) -> ReturnType (*)(ParameterTypes...)
 {
 	return reinterpret_cast<ReturnType (*)(ParameterTypes...)>(FunctionMapLookupExhaustive(reinterpret_cast<void*>(Function), Where));
 }
diff --git a/Engine/Source/Runtime/AutoRTFM/Private/FunctionMapInlines.h b/Engine/Source/Runtime/AutoRTFM/Private/FunctionMapInlines.h
--- a/Engine/Source/Runtime/AutoRTFM/Private/FunctionMapInlines.h
+++ b/Engine/Source/Runtime/AutoRTFM/Private/FunctionMapInlines.h
@@ -12,7 +12,7 @@
 namespace AutoRTFM
 {
 
-AUTORTFM_DISABLE UE_AUTORTFM_FORCEINLINE
+AUTORTFM_DISABLE_NO_SANITIZE UE_AUTORTFM_FORCEINLINE
 void* FunctionMapLookupUsingMagicPrefix(void* OpenFn)
 {
 	constexpr uint64_t UbsanMagic = 0xc105'cafe;
@@ -77,7 +77,7 @@ void* FunctionMapLookupUsingMagicPrefix(void* OpenFn)
 	return nullptr;
 }
 
-AUTORTFM_DISABLE inline void* FunctionMapLookup(void* OpenFn, const char* Where)
+AUTORTFM_DISABLE_NO_SANITIZE inline void* FunctionMapLookup(void* OpenFn, const char* Where)
 {
 	if (void* ClosedFn = FunctionMapLookupUsingMagicPrefix(OpenFn); AUTORTFM_LIKELY(ClosedFn))
 	{
@@ -88,7 +88,7 @@ AUTORTFM_DISABLE inline void* FunctionMapLookup(void* OpenFn, const char* Where)
 }
 
 template<typename TReturnType, typename... TParameterTypes>
-AUTORTFM_DISABLE auto FunctionMapLookup(TReturnType (*OpenFn)(TParameterTypes...), const char* Where) -> TReturnType (*)(TParameterTypes...)
+AUTORTFM_DISABLE_NO_SANITIZE auto FunctionMapLookup(TReturnType (*OpenFn)(TParameterTypes...), const char* Where) -> TReturnType (*)(TParameterTypes...)
 {
     return reinterpret_cast<TReturnType (*)(TParameterTypes...)>(FunctionMapLookup(reinterpret_cast<void*>(OpenFn), Where));
 }
diff --git a/Engine/Source/Runtime/AutoRTFM/Private/IntervalTree.h b/Engine/Source/Runtime/AutoRTFM/Private/IntervalTree.h
--- a/Engine/Source/Runtime/AutoRTFM/Private/IntervalTree.h
+++ b/Engine/Source/Runtime/AutoRTFM/Private/IntervalTree.h
@@ -12,7 +12,7 @@
 namespace AutoRTFM
 {
 
-struct AUTORTFM_DISABLE FIntervalTree final
+struct AUTORTFM_DISABLE_NO_SANITIZE FIntervalTree final
 {
     FIntervalTree() = default;
 
diff --git a/Engine/Source/Runtime/AutoRTFM/Private/IntrusivePool.h b/Engine/Source/Runtime/AutoRTFM/Private/IntrusivePool.h
--- a/Engine/Source/Runtime/AutoRTFM/Private/IntrusivePool.h
+++ b/Engine/Source/Runtime/AutoRTFM/Private/IntrusivePool.h
@@ -14,7 +14,7 @@ namespace AutoRTFM
 {
 
 template<typename ItemType, size_t InlineCapacity>
-class AUTORTFM_DISABLE TIntrusivePool final
+class AUTORTFM_DISABLE_NO_SANITIZE TIntrusivePool final
 {
 public:
 	using FItem = ItemType;
diff --git a/Engine/Source/Runtime/AutoRTFM/Private/Memcpy.h b/Engine/Source/Runtime/AutoRTFM/Private/Memcpy.h
--- a/Engine/Source/Runtime/AutoRTFM/Private/Memcpy.h
+++ b/Engine/Source/Runtime/AutoRTFM/Private/Memcpy.h
@@ -13,10 +13,10 @@ namespace AutoRTFM
 
 class FContext;
 
-AUTORTFM_DISABLE void* MemcpyToNew(void* Dst, const void* Src, size_t Size, FContext* Context);
-AUTORTFM_DISABLE void* Memcpy(void* Dst, const void* Src, size_t Size, FContext* Context);
-AUTORTFM_DISABLE void* Memmove(void* Dst, const void* Src, size_t Size, FContext* Context);
-AUTORTFM_DISABLE void* Memset(void* Dst, int Value, size_t Size, FContext* Context);
+AUTORTFM_DISABLE_NO_SANITIZE void* MemcpyToNew(void* Dst, const void* Src, size_t Size, FContext* Context);
+AUTORTFM_DISABLE_NO_SANITIZE void* Memcpy(void* Dst, const void* Src, size_t Size, FContext* Context);
+AUTORTFM_DISABLE_NO_SANITIZE void* Memmove(void* Dst, const void* Src, size_t Size, FContext* Context);
+AUTORTFM_DISABLE_NO_SANITIZE void* Memset(void* Dst, int Value, size_t Size, FContext* Context);
 
 } // namespace AutoRTFM
 
diff --git a/Engine/Source/Runtime/AutoRTFM/Private/OpenHashThrottler.h b/Engine/Source/Runtime/AutoRTFM/Private/OpenHashThrottler.h
--- a/Engine/Source/Runtime/AutoRTFM/Private/OpenHashThrottler.h
+++ b/Engine/Source/Runtime/AutoRTFM/Private/OpenHashThrottler.h
@@ -4,6 +4,7 @@
 
 #if (defined(__AUTORTFM) && __AUTORTFM)
 
+#include "AutoRTFM/Defines.h"
 #include "HashMap.h"
 
 namespace AutoRTFM
@@ -24,7 +25,7 @@ class FWriteLog;
 /// The percentage of transitions skipped is adjusted by the amount of time spent
 /// validating, and on a per-open basis (identified using the open's return 
 /// address).
-class FOpenHashThrottler
+class AUTORTFM_DISABLE_NO_SANITIZE FOpenHashThrottler
 {
 public:
 	using FSeconds = double;
diff --git a/Engine/Source/Runtime/AutoRTFM/Private/ScopedGuard.h b/Engine/Source/Runtime/AutoRTFM/Private/ScopedGuard.h
--- a/Engine/Source/Runtime/AutoRTFM/Private/ScopedGuard.h
+++ b/Engine/Source/Runtime/AutoRTFM/Private/ScopedGuard.h
@@ -4,11 +4,13 @@
 
 #if (defined(__AUTORTFM) && __AUTORTFM)
 
+#include "AutoRTFM/Defines.h"
+
 namespace AutoRTFM
 {
 
 template<typename T>
-class TScopedGuard
+class AUTORTFM_DISABLE_NO_SANITIZE TScopedGuard
 {
 public:
     TScopedGuard(T& Ref, const T& Value)
diff --git a/Engine/Source/Runtime/AutoRTFM/Private/Stack.h b/Engine/Source/Runtime/AutoRTFM/Private/Stack.h
--- a/Engine/Source/Runtime/AutoRTFM/Private/Stack.h
+++ b/Engine/Source/Runtime/AutoRTFM/Private/Stack.h
@@ -30,7 +30,7 @@ namespace AutoRTFM
 //                    the heap.
 //   Validation - if Disabled, do not perform validity assertions.
 template<typename T, size_t InlineCapacity, EContainerValidation Validation = EContainerValidation::Enabled>
-class TStack
+class AUTORTFM_DISABLE_NO_SANITIZE TStack
 {
 public:
 	using ElementType = T;
diff --git a/Engine/Source/Runtime/AutoRTFM/Private/StateTransitionCounter.h b/Engine/Source/Runtime/AutoRTFM/Private/StateTransitionCounter.h
--- a/Engine/Source/Runtime/AutoRTFM/Private/StateTransitionCounter.h
+++ b/Engine/Source/Runtime/AutoRTFM/Private/StateTransitionCounter.h
@@ -4,7 +4,9 @@
 
 #if (defined(__AUTORTFM) && __AUTORTFM)
 
-#include "AutoRTFM.h"
+#include "AutoRTFM/Defines.h"
+
+#include <cstdint>
 
 namespace AutoRTFM
 {
@@ -12,7 +14,7 @@ namespace AutoRTFM
 // A monotonic counter that is incremented every time a transaction changes
 // state (open -> closed, closed -> done), etc.
 // Reset when the outermost transaction is complete.
-struct AUTORTFM_DISABLE FStateTransitionCounter final
+struct AUTORTFM_DISABLE_NO_SANITIZE FStateTransitionCounter final
 {
 	using CounterType = uint64_t;
 	uint64_t Value = 0;
diff --git a/Engine/Source/Runtime/AutoRTFM/Private/ThreadID.cpp b/Engine/Source/Runtime/AutoRTFM/Private/ThreadID.cpp
--- a/Engine/Source/Runtime/AutoRTFM/Private/ThreadID.cpp
+++ b/Engine/Source/Runtime/AutoRTFM/Private/ThreadID.cpp
@@ -4,17 +4,26 @@
 
 #include "ThreadID.h"
 
-#include <thread>
+#if AUTORTFM_PLATFORM_WINDOWS
+#include "WindowsHeader.h"
+#endif
 
 namespace AutoRTFM
 {
 
-const FThreadID FThreadID::Invalid;
+const FThreadID FThreadID::Invalid{};
 
-UE_AUTORTFM_ALWAYS_OPEN
 FThreadID FThreadID::GetCurrent()
 {
+	// On Windows, call GetCurrentThreadId() directly.
+	// This avoids an open write in std::this_thread::get_id(), which can cause
+	// a stack overflow as autortfm_sanitizer_open_write() will call back into
+	// this function.
+#if AUTORTFM_PLATFORM_WINDOWS
+	return FThreadID{GetCurrentThreadId()};
+#else
 	return FThreadID{std::this_thread::get_id()};
+#endif
 }
 
 }
diff --git a/Engine/Source/Runtime/AutoRTFM/Private/ThreadID.h b/Engine/Source/Runtime/AutoRTFM/Private/ThreadID.h
--- a/Engine/Source/Runtime/AutoRTFM/Private/ThreadID.h
+++ b/Engine/Source/Runtime/AutoRTFM/Private/ThreadID.h
@@ -6,21 +6,26 @@
 
 #include "AutoRTFM.h"
 
+#include "BuildMacros.h"
+
+#if AUTORTFM_PLATFORM_WINDOWS
+#include <cstdint>
+#else
 #include <thread>
+#endif
 
 namespace AutoRTFM
 {
 
 // A unique identifier for a thread of execution.
 // Constructed as FThreadID::Invalid.
-struct FThreadID 
+struct AUTORTFM_DISABLE_NO_SANITIZE FThreadID 
 {
 	// An invalid thread identifier.
 	UE_AUTORTFM_API static const FThreadID Invalid;
 
 	// Returns the currently executing thread's unique identifier.
-	UE_AUTORTFM_API
-	static FThreadID GetCurrent();
+	UE_AUTORTFM_API static FThreadID GetCurrent();
 	
 	// Equality operator
 	inline bool operator == (const FThreadID& Other) const
@@ -34,7 +39,11 @@ struct FThreadID
 		return Value != Other.Value;
 	}
 
+#if AUTORTFM_PLATFORM_WINDOWS
+	uint32_t Value = 0;
+#else
 	std::thread::id Value;
+#endif
 };
 
 }
diff --git a/Engine/Source/Runtime/AutoRTFM/Private/Transaction.cpp b/Engine/Source/Runtime/AutoRTFM/Private/Transaction.cpp
--- a/Engine/Source/Runtime/AutoRTFM/Private/Transaction.cpp
+++ b/Engine/Source/Runtime/AutoRTFM/Private/Transaction.cpp
@@ -109,6 +109,8 @@ void FTransaction::Reset()
 	DeferredPopAllOnCommitHandlers.Reset();
 	DeferredPopAllOnAbortHandlers.Reset();
 
+	AUTORTFM_SANITIZER_ONLY(ClosedWrites.Reset());
+
 	AUTORTFM_ASSERT(IsFresh());
 }
 
@@ -125,7 +127,11 @@ bool FTransaction::IsFresh() const
 		&& DeferredPopOnCommitHandlers.IsEmpty()
 		&& DeferredPopOnAbortHandlers.IsEmpty()
 		&& DeferredPopAllOnCommitHandlers.IsEmpty()
-		&& DeferredPopAllOnAbortHandlers.IsEmpty();
+		&& DeferredPopAllOnAbortHandlers.IsEmpty()
+#if AUTORTFM_SANITIZER
+		&& ClosedWrites.IsEmpty()
+#endif
+		;
 }
 
 void FTransaction::AbortWithoutThrowing()
@@ -339,7 +345,9 @@ void FTransaction::Undo()
 		{
 			memcpy(Entry.LogicalAddress, Entry.Data, Entry.Size);
 		}
-    }
+	}
+
+	AUTORTFM_SANITIZER_ONLY(ClosedWrites.Reset());
 
 	AUTORTFM_VERBOSE("Undone a transaction!");
 }
@@ -373,6 +381,13 @@ void FTransaction::CommitNested()
 		}
 
 		Parent->WriteLog.Push(Write);
+
+#if AUTORTFM_SANITIZER
+		if ((Write.Flags & EWriteFlags::NoValidation) != EWriteFlags::NoValidation)
+		{
+			Parent->ClosedWrites.Insert(Write.LogicalAddress, Write.Size);
+		}
+#endif
 	}
 
 	// Similarly we need to merge the hit-sets to the parent, skipping those
@@ -419,11 +434,11 @@ void FTransaction::CommitNested()
 	}
 	DeferredPopAllOnAbortHandlers.Reset();
 
-    Parent->CommitTasks.AddAll(std::move(CommitTasks));
-    Parent->PreAbortTasks.AddAll(std::move(PreAbortTasks));
-    Parent->AbortTasks.AddAll(std::move(AbortTasks));
+	Parent->CommitTasks.AddAll(std::move(CommitTasks));
+	Parent->PreAbortTasks.AddAll(std::move(PreAbortTasks));
+	Parent->AbortTasks.AddAll(std::move(AbortTasks));
 
-    Parent->NewMemoryTracker.Merge(NewMemoryTracker);
+	Parent->NewMemoryTracker.Merge(NewMemoryTracker);
 }
 
 bool FTransaction::AttemptToCommitOuterNest()
@@ -434,20 +449,22 @@ bool FTransaction::AttemptToCommitOuterNest()
 	Context->DumpState();
 	AUTORTFM_VERBOSE("Running commit tasks...");
 
-    PreAbortTasks.Reset();
-    AbortTasks.Reset();
+	PreAbortTasks.Reset();
+	AbortTasks.Reset();
 
-    CommitTasks.RemoveEachForward([] (TTask<void()>& Task)
-    { 
-        Task();
-    });
+	AUTORTFM_SANITIZER_ONLY(ClosedWrites.Reset());
 
-    CompletionTasks->RemoveEachForward([] (TTask<void()>& Task)
-    { 
-        Task();
-    });
+	CommitTasks.RemoveEachForward([] (TTask<void()>& Task)
+	{ 
+		Task();
+	});
+
+	CompletionTasks->RemoveEachForward([] (TTask<void()>& Task)
+	{ 
+		Task();
+	});
 
-    return true;
+	return true;
 }
 
 void FTransaction::SetOpenActiveValidatorEnabled(EMemoryValidationLevel NewMemoryValidationLevel, const void* ReturnAddress)
@@ -628,7 +645,7 @@ void FTransaction::RecordWriteHash()
 
 #if PLATFORM_WINDOWS
 // Used to identify the specific open write that causes a memory validation hash mismatch.
-struct AUTORTFM_DISABLE FTransaction::FBadWriteIdentifier
+struct AUTORTFM_DISABLE_NO_SANITIZE FTransaction::FBadWriteIdentifier
 {
 	// Constructor
 	// OpenCounter is the transaction state transition counter at the transition
@@ -839,7 +856,7 @@ void FTransaction::ValidateWriteHash()
 
 	if (OldHash != NewHash)
 	{
-		std::string Message = 
+		std::string Message =
 			"Memory modified in a transaction was also modified in an call to AutoRTFM::Open().\n"
 			"This may lead to memory corruption if the transaction is aborted.";
 
diff --git a/Engine/Source/Runtime/AutoRTFM/Private/Transaction.h b/Engine/Source/Runtime/AutoRTFM/Private/Transaction.h
--- a/Engine/Source/Runtime/AutoRTFM/Private/Transaction.h
+++ b/Engine/Source/Runtime/AutoRTFM/Private/Transaction.h
@@ -22,7 +22,7 @@ namespace AutoRTFM
 {
 class FContext;
 
-class AUTORTFM_DISABLE FTransaction final
+class AUTORTFM_DISABLE_NO_SANITIZE FTransaction final
 {
 	friend class TIntrusivePool<FTransaction, 16>;
 
@@ -159,6 +159,10 @@ public:
 	// Returns true if the LogicalAddress is within the stack of the transaction.
 	inline bool IsOnStack(const void* LogicalAddress) const;
 
+#if AUTORTFM_SANITIZER
+	inline void SanitizerOpenWrite(void* Ptr, size_t Size);
+#endif
+
 private:
 	using FWriteHash = FWriteLog::FHash;
 	using FTaskArray = TTaskArray<TTask<void()>>;
@@ -257,6 +261,9 @@ private:
 
 	FHitSet HitSet;
 	FIntervalTree NewMemoryTracker;
+#if AUTORTFM_SANITIZER
+	FIntervalTree ClosedWrites;
+#endif
 	FWriteLog WriteLog;
 	TStatStorage<uint64_t> StatDepth = 1;
 	FWriteLogSnapshot WriteLogSnapshot;
diff --git a/Engine/Source/Runtime/AutoRTFM/Private/TransactionInlines.h b/Engine/Source/Runtime/AutoRTFM/Private/TransactionInlines.h
--- a/Engine/Source/Runtime/AutoRTFM/Private/TransactionInlines.h
+++ b/Engine/Source/Runtime/AutoRTFM/Private/TransactionInlines.h
@@ -107,6 +107,8 @@ void FTransaction::RecordWrite(void* LogicalAddress, size_t Size, EWriteFlags Fl
 		.Flags = Flags
 	};
 
+	AutoRTFM::TScopedGuard<bool> RecursionGuard(bIsInAllocateFn, true);
+	
 	switch (Entry.Size)
 	{
 		case 8:
@@ -122,7 +124,13 @@ void FTransaction::RecordWrite(void* LogicalAddress, size_t Size, EWriteFlags Fl
 			WriteLog.Push(Entry);
 			break;
 	}
-	
+
+#if AUTORTFM_SANITIZER
+	if ((Flags & EWriteFlags::NoValidation) != EWriteFlags::NoValidation)
+	{
+		ClosedWrites.Insert(LogicalAddress, Size);
+	}
+#endif
 }
 
 template<unsigned SIZE> AUTORTFM_NO_ASAN UE_AUTORTFM_FORCEINLINE void FTransaction::RecordWrite(void* LogicalAddress)
@@ -178,7 +186,11 @@ template<unsigned SIZE> AUTORTFM_NO_ASAN UE_AUTORTFM_FORCENOINLINE void FTransac
 
 	Stats.Collect<EStatsKind::NewMemoryTrackerMiss>();
 
+	AutoRTFM::TScopedGuard<bool> RecursionGuard(bIsInAllocateFn, true);
+
 	WriteLog.PushSmall<SIZE>(static_cast<std::byte*>(LogicalAddress), static_cast<std::byte*>(LogicalAddress));
+
+	AUTORTFM_SANITIZER_ONLY(ClosedWrites.Insert(LogicalAddress, SIZE));
 }
 
 UE_AUTORTFM_FORCEINLINE void FTransaction::DidAllocate(void* LogicalAddress, const size_t Size)
@@ -201,9 +213,9 @@ UE_AUTORTFM_FORCEINLINE void FTransaction::DidFree(void* LogicalAddress)
 	// is new memory and we should be worried.
 	if (!bIsInAllocateFn)
 	{
-		AutoRTFM::TScopedGuard<bool> RecursionGuard(bIsInAllocateFn, true);
 		if (NewMemoryTracker.Contains(LogicalAddress, 1))
 		{
+			AutoRTFM::TScopedGuard<bool> RecursionGuard(bIsInAllocateFn, true);
 			ForTheRuntime::ReportAutoRTFMHazard(ForTheRuntime::EHazardType::OpenFree);
 			AUTORTFM_REPORT_ERROR("Transactional memory was freed from the open.");
 		}
@@ -328,6 +340,45 @@ UE_AUTORTFM_FORCEINLINE void FTransaction::CollectStats() const
     Stats.Collect<EStatsKind::AverageHitSetCapacity>(HitSet.GetCapacity());
 }
 
+#if AUTORTFM_SANITIZER
+void FTransaction::SanitizerOpenWrite(void* Ptr, size_t Size)
+{
+	if (bIsInAllocateFn)
+	{
+		// Open write was performed while we were allocating internal AutoRTFM
+		// memory. Ignore.
+		return;
+	}
+
+	if (CurrentMemoryValidationLevel == EMemoryValidationLevel::Disabled)
+	{
+		return;
+	}
+	
+	for (FTransaction* Transaction = this; Transaction; Transaction = Transaction->Parent)
+	{
+		// TODO(SOL-8801): FIntervalTree::Contains() may incorrectly return
+		// false if the queried interval starts outside of an interval added by
+		// FIntervalTree::Insert() but still overlaps an inserted interval.
+		// Fix this.
+		if (AUTORTFM_UNLIKELY(Transaction->ClosedWrites.Contains(Ptr, Size)))
+		{
+			const char* Message =
+				"AutoRTFMSan: 0x%x bytes at %p was modified in the closed and now the open.\n"
+				"This may lead to memory corruption if the transaction is aborted.";
+			if (CurrentMemoryValidationLevel == EMemoryValidationLevel::Error)
+			{
+				AUTORTFM_FATAL(Message, static_cast<int>(Size), Ptr);
+			}
+			else
+			{
+				AUTORTFM_WARN(Message, static_cast<int>(Size), Ptr);
+			}
+		}
+	}
+}
+#endif
+
 } // namespace AutoRTFM
 
 #undef AUTORTFM_CHECK_ASAN_FAKE_STACKS
diff --git a/Engine/Source/Runtime/AutoRTFM/Private/Utils.h b/Engine/Source/Runtime/AutoRTFM/Private/Utils.h
--- a/Engine/Source/Runtime/AutoRTFM/Private/Utils.h
+++ b/Engine/Source/Runtime/AutoRTFM/Private/Utils.h
@@ -290,6 +290,12 @@ inline constexpr T Lerp(T Start, T End, T Fraction)
 		AutoRTFM::DoExpect(Lambda); \
 	}
 
+#if AUTORTFM_SANITIZER
+#define AUTORTFM_SANITIZER_ONLY(...) __VA_ARGS__
+#else
+#define AUTORTFM_SANITIZER_ONLY(...) static_assert(true) // require semicolon
+#endif
+	
 #if defined(__has_feature)
 #if __has_feature(address_sanitizer)
 #define AUTORTFM_NO_ASAN [[clang::no_sanitize("address")]]
diff --git a/Engine/Source/Runtime/AutoRTFM/Private/WriteLog.h b/Engine/Source/Runtime/AutoRTFM/Private/WriteLog.h
--- a/Engine/Source/Runtime/AutoRTFM/Private/WriteLog.h
+++ b/Engine/Source/Runtime/AutoRTFM/Private/WriteLog.h
@@ -16,7 +16,7 @@
 
 namespace AutoRTFM
 {
-	struct FWriteLogEntry final
+	struct AUTORTFM_DISABLE_NO_SANITIZE FWriteLogEntry final
 	{
 		// The address of the write.
 		std::byte* LogicalAddress = nullptr;
@@ -47,7 +47,7 @@ namespace AutoRTFM
 	// FWriteLog holds an ordered list of write records which can be iterated
 	// forwards and backwards.
 	// Ensure changes to this class are kept in sync with AutoRTFM.natvis.
-	class FWriteLog final
+	class AUTORTFM_DISABLE_NO_SANITIZE FWriteLog final
 	{
 		enum EAppendMode
 		{
