[AutoRTFM] Reduce the number of heap allocations performed by FIntervalTree by over 4x.

FIntervalTree uses structure-of-arrays for the node data. This was stored in 5 vectors, 4 of which would always heap-resize together.
Continue to use structure-of-arrays, but heap-allocate a single chunk of memory, and only resize that once.
Saves time performing heap allocations, and reduces the number of needs-to-resize checks by 5x.

Delete the now unused FBitStack.

#Synced-CL 48856386

diff --git a/Engine/Source/Programs/AutoRTFMTests/Private/BitStackTests.cpp b/Engine/Source/Programs/AutoRTFMTests/Private/BitStackTests.cpp
deleted file mode 100644
diff --git a/Engine/Source/Runtime/AutoRTFM/Private/BitStack.h b/Engine/Source/Runtime/AutoRTFM/Private/BitStack.h
deleted file mode 100644
diff --git a/Engine/Source/Runtime/AutoRTFM/Private/IntervalTree.h b/Engine/Source/Runtime/AutoRTFM/Private/IntervalTree.h
--- a/Engine/Source/Runtime/AutoRTFM/Private/IntervalTree.h
+++ b/Engine/Source/Runtime/AutoRTFM/Private/IntervalTree.h
@@ -4,106 +4,103 @@
 
 #if (defined(__AUTORTFM) && __AUTORTFM)
 
-#include "BitStack.h"
+#include "ExternAPI.h"
 #include "ContainerValidation.h"
 #include "Stack.h"
 #include "Utils.h"
+#include <cstring>
 
 namespace AutoRTFM
 {
 
 struct AUTORTFM_DISABLE_NO_SANITIZE FIntervalTree final
 {
-    FIntervalTree() = default;
+	FIntervalTree() = default;
 
-    FIntervalTree(const FIntervalTree&) = delete;
-    FIntervalTree& operator=(const FIntervalTree&) = delete;
+	FIntervalTree(const FIntervalTree&) = delete;
+	FIntervalTree& operator=(const FIntervalTree&) = delete;
 
 	UE_AUTORTFM_FORCENOINLINE bool Insert(const void* const Address, const size_t Size)
-    {        
-        FRange NewRange(Address, Size);
-        return Insert(std::move(NewRange));
-    }
-
-    UE_AUTORTFM_FORCENOINLINE bool Contains(const void* const Address, const size_t Size) const
-    {
-        const FRange NewRange(Address, Size);
-
-        if (AUTORTFM_UNLIKELY(IntervalTreeNodeIndexNone == Root))
-        {
-            return false;
-        }
-
-        FIntervalTreeNodeIndex Current = Root;
-
-        do
-        {
-            const FRange Range = NodeRanges[Current];
-
-            // This check does not need to prove that NewRange is entirely
-            // enclosed within Range, because if any byte of NewRange was in the
-            // original Range then it **must** already have been new memory.
-            if ((NewRange.Start < Range.End) && (Range.Start < NewRange.End))
-            {
-                return true;
-            }
-
-			const StackType<FIntervalTreeNodeIndex>& Next = (NewRange.Start < Range.End) ? NodeLefts : NodeRights;
+	{        
+		FRange NewRange(Address, Size);
+		return Insert(std::move(NewRange));
+	}
+
+	UE_AUTORTFM_FORCENOINLINE bool Contains(const void* const Address, const size_t Size) const
+	{
+		const FRange NewRange(Address, Size);
+
+		if (AUTORTFM_UNLIKELY(IntervalTreeNodeIndexNone == Root))
+		{
+			return false;
+		}
+
+		FIntervalTreeNodeIndex Current = Root;
+
+		do
+		{
+			const FRange Range = Nodes.Ranges[Current];
+
+			// This check does not need to prove that NewRange is entirely
+			// enclosed within Range, because if any byte of NewRange was in the
+			// original Range then it **must** already have been new memory.
+			if ((NewRange.Start < Range.End) && (Range.Start < NewRange.End))
+			{
+				return true;
+			}
+
+			const FIntervalTreeNodeIndex* Next = (NewRange.Start < Range.End) ? Nodes.Lefts : Nodes.Rights;
 			Current = Next[Current];
-        } while (IntervalTreeNodeIndexNone != Current);
+		} while (IntervalTreeNodeIndexNone != Current);
 
-        return false;
-    }
+		return false;
+	}
 
-    bool IsEmpty() const
-    {
-        return IntervalTreeNodeIndexNone == Root;
-    }
+	bool IsEmpty() const
+	{
+		return IntervalTreeNodeIndexNone == Root;
+	}
 
-    void Reset()
-    {
-        Root = IntervalTreeNodeIndexNone;
-		NodeRanges.Reset();
-		NodeLefts.Reset();
-		NodeRights.Reset();
-		NodeParents.Reset();
-		NodeColors.Reset();
-    }
+	void Reset()
+	{
+		Root = IntervalTreeNodeIndexNone;
+		Nodes.Reset();
+	}
 
-    void Merge(const FIntervalTree& Other)
-    {
-        if (Other.IsEmpty())
-        {
-            return;
-        }
+	void Merge(const FIntervalTree& Other)
+	{
+		if (Other.IsEmpty())
+		{
+			return;
+		}
 
-		StackType<FIntervalTreeNodeIndex> ToProcess;
+		TStack<FIntervalTreeNodeIndex, 32, EContainerValidation::Disabled> ToProcess;
 
-        ToProcess.Push(Other.Root);
+		ToProcess.Push(Other.Root);
 
-        do
-        {
-            const FIntervalTreeNodeIndex Current = ToProcess.Back();
-            ToProcess.Pop();
+		do
+		{
+			const FIntervalTreeNodeIndex Current = ToProcess.Back();
+			ToProcess.Pop();
 
-            FRange Range = Other.NodeRanges[Current];
+			FRange Range = Other.Nodes.Ranges[Current];
 
-            Insert(std::move(Range));
+			Insert(std::move(Range));
 
-            const FIntervalTreeNodeIndex Left = Other.NodeLefts[Current];
-            const FIntervalTreeNodeIndex Right = Other.NodeRights[Current];
+			const FIntervalTreeNodeIndex Left = Other.Nodes.Lefts[Current];
+			const FIntervalTreeNodeIndex Right = Other.Nodes.Rights[Current];
 
-            if (IntervalTreeNodeIndexNone != Left)
-            {
-                ToProcess.Push(Left);
-            }
+			if (IntervalTreeNodeIndexNone != Left)
+			{
+				ToProcess.Push(Left);
+			}
 
-            if (IntervalTreeNodeIndexNone != Right)
-            {
-                ToProcess.Push(Right);
-            }
-        } while (!ToProcess.IsEmpty());
-    }
+			if (IntervalTreeNodeIndexNone != Right)
+			{
+				ToProcess.Push(Right);
+			}
+		} while (!ToProcess.IsEmpty());
+	}
 
 	// Calls Callback with each interval in the tree, in order of address.
 	// Callback must have the signature: Callback(uintptr_t Start, uintptr_t End)
@@ -117,22 +114,18 @@ struct AUTORTFM_DISABLE_NO_SANITIZE FIntervalTree final
 	}
 
 private:
-	static constexpr size_t InlineArraySize = 8;
-	template<typename T> using StackType = TStack<T, InlineArraySize, EContainerValidation::Disabled>;
-	using BitStackType = TBitStack<InlineArraySize, EContainerValidation::Disabled>;
-
-    struct FRange final
-    {
-        FRange(const void* const Address, const size_t Size) :
-            Start(reinterpret_cast<uintptr_t>(Address)), End(reinterpret_cast<uintptr_t>(Address) + Size) {}
+	struct FRange final
+	{
+		FRange(const void* const Address, const size_t Size) :
+			Start(reinterpret_cast<uintptr_t>(Address)), End(reinterpret_cast<uintptr_t>(Address) + Size) {}
 
-        uintptr_t Start;
-        uintptr_t End;
-    };
+		uintptr_t Start;
+		uintptr_t End;
+	};
 
-    using FIntervalTreeNodeIndex = uint32_t;
+	using FIntervalTreeNodeIndex = uint32_t;
 
-    static constexpr FIntervalTreeNodeIndex IntervalTreeNodeIndexNone = UINT32_MAX;
+	static constexpr FIntervalTreeNodeIndex IntervalTreeNodeIndexNone = UINT32_MAX;
 
 	enum EColor : bool
 	{
@@ -140,24 +133,162 @@ private:
 		Black = true
 	};
 
-    FIntervalTreeNodeIndex Root = IntervalTreeNodeIndexNone;
-	StackType<FRange> NodeRanges;
-	StackType<FIntervalTreeNodeIndex> NodeLefts;
-	StackType<FIntervalTreeNodeIndex> NodeRights;
-	StackType<FIntervalTreeNodeIndex> NodeParents;
-	BitStackType NodeColors;
+	struct FNode final
+	{
+		FRange Range;
+		FIntervalTreeNodeIndex Left = IntervalTreeNodeIndexNone;
+		FIntervalTreeNodeIndex Right = IntervalTreeNodeIndexNone;
+		FIntervalTreeNodeIndex Parent = IntervalTreeNodeIndexNone;
+		EColor Color;
+	};
+
+	// A list of FNodes, stored as a structure-of-arrays.
+	struct FNodes final
+	{
+		using ColorWord = uint64_t;
+		using RangeType = decltype(FNode::Range);
+		using LeftType = decltype(FNode::Left);
+		using RightType = decltype(FNode::Right);
+		using ParentType = decltype(FNode::Parent);
+		static constexpr size_t NumColorWordBits = sizeof(ColorWord) * 8;
+		static constexpr FIntervalTreeNodeIndex InitialCapacity = 32;
+
+		RangeType* Ranges = nullptr;
+		LeftType* Lefts = nullptr;
+		RightType* Rights = nullptr;
+		ParentType* Parents = nullptr;
+		ColorWord* ColorWords = nullptr;
+		FIntervalTreeNodeIndex Size = 0;
+		FIntervalTreeNodeIndex Capacity = 0;
+
+		FNodes() = default;
+		
+		~FNodes()
+		{
+			Reset();
+		}
+
+		// Clears the node list to an empty state.
+		void Reset()
+		{
+			AutoRTFM::Free(Memory);
+			Ranges = nullptr;
+			Lefts = nullptr;
+			Rights = nullptr;
+			Parents = nullptr;
+			ColorWords = nullptr;
+			Size = 0;
+			Capacity = 0;
+			Memory = nullptr;
+		}
+
+		// Pushes a new node to the end of the list.
+		// Returns the node index.
+		FIntervalTreeNodeIndex Push(FNode Node)
+		{
+			if (Size >= Capacity)
+			{
+				Resize(std::max(Capacity * 2, InitialCapacity));
+			}
+
+			Ranges[Size] = Node.Range;
+			Lefts[Size] = Node.Left;
+			Rights[Size] = Node.Right;
+			Parents[Size] = Node.Parent;
+			SetColor(Size, Node.Color);
+
+			return static_cast<FIntervalTreeNodeIndex>(Size++);
+		}
+
+		void SetColor(FIntervalTreeNodeIndex Index, EColor Color)
+		{
+			ColorWord const Bit = static_cast<ColorWord>(1) << (Index & (NumColorWordBits - 1));
+			ColorWord& Word = ColorWords[Index / NumColorWordBits];
+			if (Color == EColor::Black)
+			{
+				Word |= Bit;
+			}
+			else
+			{
+				Word &= ~Bit;
+			}
+		}
+
+		EColor Color(size_t Index) const
+		{
+			ColorWord const Bit = static_cast<ColorWord>(1) << (Index & (NumColorWordBits - 1));
+			ColorWord& Word = ColorWords[Index / NumColorWordBits];
+			return (Word & Bit) ? EColor::Black : EColor::Red;
+		}
+
+	private:
+		std::byte* Memory = nullptr;
+
+		FNodes(const FNodes&) = delete;
+		FNodes& operator=(const FNodes&) = delete;
+	
+		void Resize(FIntervalTreeNodeIndex NewCapacity)
+		{
+			size_t TotalSize = 0;
+			size_t const RangesOffset = SubAllocate<RangeType>(TotalSize, NewCapacity);
+			size_t const LeftsOffset = SubAllocate<LeftType>(TotalSize, NewCapacity);
+			size_t const RightsOffset = SubAllocate<RightType>(TotalSize, NewCapacity);
+			size_t const ParentsOffset = SubAllocate<ParentType>(TotalSize, NewCapacity);
+			size_t const ColorWordsOffset = SubAllocate<ColorWord>(TotalSize, NumColorWords(NewCapacity));
+			auto* const NewMemory = reinterpret_cast<std::byte*>(AutoRTFM::Allocate(TotalSize, alignof(FNode)));
+			auto* const NewRanges = std::launder(reinterpret_cast<RangeType*>(NewMemory + RangesOffset));
+			auto* const NewLefts = std::launder(reinterpret_cast<LeftType*>(NewMemory + LeftsOffset));
+			auto* const NewRights = std::launder(reinterpret_cast<RightType*>(NewMemory + RightsOffset));
+			auto* const NewParents = std::launder(reinterpret_cast<ParentType*>(NewMemory + ParentsOffset));
+			auto* const NewColorWords = std::launder(reinterpret_cast<ColorWord*>(NewMemory + ColorWordsOffset));
+			if (size_t const NumToCopy = std::min(Size, NewCapacity))
+			{
+				memcpy(NewRanges, Ranges, sizeof(RangeType) * NumToCopy);
+				memcpy(NewLefts, Lefts, sizeof(LeftType) * NumToCopy);
+				memcpy(NewRights, Rights, sizeof(RightType) * NumToCopy);
+				memcpy(NewParents, Parents, sizeof(ParentType) * NumToCopy);
+				memcpy(NewColorWords, ColorWords, sizeof(ColorWord) * NumColorWords(NumToCopy));
+			}
+
+			AutoRTFM::Free(Memory);
+
+			Ranges = NewRanges;
+			Lefts = NewLefts;
+			Rights = NewRights;
+			Parents = NewParents;
+			ColorWords = NewColorWords;
+			Capacity = NewCapacity;
+			Memory = NewMemory;
+		}
+
+		template<typename T>
+		static size_t SubAllocate(size_t& Offset, size_t Count)
+		{
+			size_t Base = AlignUp(Offset, alignof(T));
+			Offset = Base + sizeof(T) * Count;
+			return Base;
+		}
+
+		static constexpr size_t NumColorWords(size_t Count)
+		{
+			return (Count + NumColorWordBits - 1) / NumColorWordBits;
+		}
+	};
+
+	FIntervalTreeNodeIndex Root = IntervalTreeNodeIndexNone;
+	FNodes Nodes;
 
 	static constexpr bool bExtraDebugging = false;
 
 	template<typename CALLBACK>
 	void ForEach(FIntervalTreeNodeIndex Index, CALLBACK&& Callback) const
 	{
-		if (FIntervalTreeNodeIndex Left = NodeLefts[Index]; Left != IntervalTreeNodeIndexNone)
+		if (FIntervalTreeNodeIndex Left = Nodes.Lefts[Index]; Left != IntervalTreeNodeIndexNone)
 		{
 			ForEach(Left, Callback);
 		}
-		Callback(NodeRanges[Index].Start, NodeRanges[Index].End);
-		if (FIntervalTreeNodeIndex Right = NodeRights[Index]; Right != IntervalTreeNodeIndexNone)
+		Callback(Nodes.Ranges[Index].Start, Nodes.Ranges[Index].End);
+		if (FIntervalTreeNodeIndex Right = Nodes.Rights[Index]; Right != IntervalTreeNodeIndexNone)
 		{
 			ForEach(Right, Callback);
 		}
@@ -165,237 +296,225 @@ private:
 
 	UE_AUTORTFM_FORCEINLINE FIntervalTreeNodeIndex AddNode(FRange&& NewRange)
 	{
-		const FIntervalTreeNodeIndex Index = static_cast<FIntervalTreeNodeIndex>(NodeRanges.Num());
-        NodeRanges.Push(NewRange);
-		NodeLefts.Push(IntervalTreeNodeIndexNone);
-		NodeRights.Push(IntervalTreeNodeIndexNone);
-		NodeParents.Push(IntervalTreeNodeIndexNone);
-		NodeColors.Push(EColor::Black);
-		return Index;
+		return Nodes.Push(FNode{.Range = NewRange, .Color = EColor::Black});
 	}
 
 	UE_AUTORTFM_FORCEINLINE FIntervalTreeNodeIndex AddNode(FRange&& NewRange, FIntervalTreeNodeIndex Parent)
 	{
-		const FIntervalTreeNodeIndex Index = static_cast<FIntervalTreeNodeIndex>(NodeRanges.Num());
-        NodeRanges.Push(NewRange);
-		NodeLefts.Push(IntervalTreeNodeIndexNone);
-		NodeRights.Push(IntervalTreeNodeIndexNone);
-		NodeParents.Push(Parent);
-		NodeColors.Push(EColor::Red);
-		return Index;
+		return Nodes.Push(FNode{.Range = NewRange, .Parent = Parent, .Color = EColor::Red});
 	}
 
-    UE_AUTORTFM_FORCENOINLINE bool Insert(FRange&& NewRange)
-    {
-        if (AUTORTFM_UNLIKELY(IntervalTreeNodeIndexNone == Root))
-        {
-            AUTORTFM_ASSERT(NodeRanges.IsEmpty());
+	UE_AUTORTFM_FORCENOINLINE bool Insert(FRange&& NewRange)
+	{
+		if (AUTORTFM_UNLIKELY(IntervalTreeNodeIndexNone == Root))
+		{
+			AUTORTFM_ASSERT(Nodes.Size == 0);
 			Root = AddNode(std::move(NewRange));
-            return true;
-        }
+			return true;
+		}
 
-        FIntervalTreeNodeIndex Current = Root;
+		FIntervalTreeNodeIndex Current = Root;
 
-        for(;;)
-        {
-            const FRange Range = NodeRanges[Current];
+		for(;;)
+		{
+			const FRange Range = Nodes.Ranges[Current];
 
-            if (AUTORTFM_UNLIKELY((NewRange.Start < Range.End) && (Range.Start < NewRange.End)))
-            {
-                return false;
-            }
+			if (AUTORTFM_UNLIKELY((NewRange.Start < Range.End) && (Range.Start < NewRange.End)))
+			{
+				return false;
+			}
 
-            if (NewRange.Start < Range.Start)
-            {
-                if (AUTORTFM_UNLIKELY(NewRange.End == Range.Start))
-                {
-                    AUTORTFM_ASSERT(NewRange.Start < Range.Start);
+			if (NewRange.Start < Range.Start)
+			{
+				if (AUTORTFM_UNLIKELY(NewRange.End == Range.Start))
+				{
+					AUTORTFM_ASSERT(NewRange.Start < Range.Start);
 
 					// We can just modify the existing node in place.
-					NodeRanges[Current].Start = NewRange.Start;
-                    return true;
-                }
-                else if (IntervalTreeNodeIndexNone == NodeLefts[Current])
-                {
+					Nodes.Ranges[Current].Start = NewRange.Start;
+					return true;
+				}
+				else if (IntervalTreeNodeIndexNone == Nodes.Lefts[Current])
+				{
 					const FIntervalTreeNodeIndex Index = AddNode(std::move(NewRange), Current);
-					NodeLefts[Current] = Index;
-                    Current = Index;
-                    break;
-                }
-
-                Current = NodeLefts[Current];
-            }
-            else
-            {
-                if (AUTORTFM_UNLIKELY(NewRange.Start == Range.End))
-                {
-                    AUTORTFM_ASSERT(NewRange.End > Range.End);
+					Nodes.Lefts[Current] = Index;
+					Current = Index;
+					break;
+				}
+
+				Current = Nodes.Lefts[Current];
+			}
+			else
+			{
+				if (AUTORTFM_UNLIKELY(NewRange.Start == Range.End))
+				{
+					AUTORTFM_ASSERT(NewRange.End > Range.End);
 
 					// We can just modify the existing node in place.
-					NodeRanges[Current].End = NewRange.End;
-                    return true;
-                }
-                else if (IntervalTreeNodeIndexNone == NodeRights[Current])
-                {
-                    const FIntervalTreeNodeIndex Index = AddNode(std::move(NewRange), Current);
-					NodeRights[Current] = Index;
-                    Current = Index;
-                    break;
-                }
-
-                Current = NodeRights[Current];
-            }
-
-            AUTORTFM_ASSERT(Root != Current);
-        }
-
-        auto IsBlack = [this](FIntervalTreeNodeIndex Index)
-        {
-			return (IntervalTreeNodeIndexNone == Index) || (EColor::Black == NodeColors[Index]);
-        };
-
-        for(;;)
-        {
-            FIntervalTreeNodeIndex Parent = NodeParents[Current];
-
-            UE_AUTORTFM_ASSUME(Current != Parent);
-
-            // The root will always have a black parent, so this check covers both.
-            if (Parent == Root)
-            {
-                NodeColors[Parent] = EColor::Black;
-                break;
-            }
-            else if (IsBlack(Parent))
-            {
-                break;
-            }
-
-            const FIntervalTreeNodeIndex GrandParent = NodeParents[Parent];
-
-            UE_AUTORTFM_ASSUME((Parent != GrandParent) && (Current != GrandParent));
-
-            const bool bParentIsLeft = NodeLefts[GrandParent] == Parent;
-
-            // The uncle is the other node of our parent.
-            const FIntervalTreeNodeIndex Uncle = bParentIsLeft ? NodeRights[GrandParent] : NodeLefts[GrandParent];
-
-            UE_AUTORTFM_ASSUME((GrandParent != Uncle) && (Parent != Uncle) && (Current != Uncle));
-
-            if (!IsBlack(Uncle))
-            {
-                AUTORTFM_ASSERT(IsBlack(GrandParent));
-                NodeColors[Parent] = EColor::Black;
-                NodeColors[Uncle] = EColor::Black;
-
-                if (GrandParent == Root)
-                {
-                    break;
-                }
-                else
-                {
-					NodeColors[GrandParent] = EColor::Red;
-                }
-
-                Current = GrandParent;
-                continue;
-            }
-
-            // Our uncle is black, so we need to swizzle around.
-            const bool bCurrentIsLeft = NodeLefts[Parent] == Current;
-
-            if (bParentIsLeft)
-            {
-                if (!bCurrentIsLeft)
-                {
-                    NodeRights[Parent] = NodeLefts[Current];
-
-                    if (IntervalTreeNodeIndexNone != NodeRights[Parent])
-                    {
-                        NodeParents[NodeRights[Parent]] = Parent;
-                    }
-
-                    NodeParents[Parent] = Current;
-                    NodeParents[Current] = GrandParent;
-                    NodeLefts[Current] = Parent;
-                    NodeLefts[GrandParent] = Current;
-                    std::swap(Parent, Current);
-                }
-
-                NodeParents[Parent] = NodeParents[GrandParent];
-                NodeLefts[GrandParent] = NodeRights[Parent];
-
-                if (IntervalTreeNodeIndexNone != NodeLefts[GrandParent])
-                {
-                    NodeParents[NodeLefts[GrandParent]] = GrandParent;
-                }
-
-                NodeRights[Parent] = GrandParent;
-                NodeParents[GrandParent] = Parent;
-
-				const bool bGrandParentIsBlack = NodeColors[GrandParent];
-				NodeColors[GrandParent] = NodeColors[Parent];
-				NodeColors[Parent] = bGrandParentIsBlack;
-            }
-            else
-            {
-                if (bCurrentIsLeft)
-                {
-                    NodeLefts[Parent] = NodeRights[Current];
-                    
-                    if (IntervalTreeNodeIndexNone != NodeLefts[Parent])
-                    {
-                        NodeParents[NodeLefts[Parent]] = Parent;
-                    }
-
-                    NodeParents[Parent] = Current;
-                    NodeParents[Current] = GrandParent;
-                    NodeRights[Current] = Parent;
-                    NodeRights[GrandParent] = Current;
-                    std::swap(Parent, Current);
-                }
-
-                NodeParents[Parent] = NodeParents[GrandParent];
-                NodeRights[GrandParent] = NodeLefts[Parent];
-
-                if (IntervalTreeNodeIndexNone != NodeRights[GrandParent])
-                {
-                    NodeParents[NodeRights[GrandParent]] = GrandParent;
-                }
-
-                NodeLefts[Parent] = GrandParent;
-                NodeParents[GrandParent] = Parent;
-
-				const bool bGrandParentIsBlack = NodeColors[GrandParent];
-				NodeColors[GrandParent] = NodeColors[Parent];
-				NodeColors[Parent] = bGrandParentIsBlack;
-            }
-
-            if (IntervalTreeNodeIndexNone == NodeParents[Parent])
-            {
-                Root = Parent;
-            }
-            else
-            {
-				const FIntervalTreeNodeIndex GreatGrandParent = NodeParents[Parent];
-
-                if (GrandParent == NodeLefts[GreatGrandParent])
-                {
-					NodeLefts[GreatGrandParent] = Parent;
-                }
-                else
-                {
-					NodeRights[GreatGrandParent] = Parent;
-                }
-            }
-
-            break;
-        }
+					Nodes.Ranges[Current].End = NewRange.End;
+					return true;
+				}
+				else if (IntervalTreeNodeIndexNone == Nodes.Rights[Current])
+				{
+					const FIntervalTreeNodeIndex Index = AddNode(std::move(NewRange), Current);
+					Nodes.Rights[Current] = Index;
+					Current = Index;
+					break;
+				}
+
+				Current = Nodes.Rights[Current];
+			}
+
+			AUTORTFM_ASSERT(Root != Current);
+		}
+
+		auto IsBlack = [this](FIntervalTreeNodeIndex Index)
+		{
+			return (IntervalTreeNodeIndexNone == Index) || (EColor::Black == Nodes.Color(Index));
+		};
+
+		for(;;)
+		{
+			FIntervalTreeNodeIndex Parent = Nodes.Parents[Current];
+
+			UE_AUTORTFM_ASSUME(Current != Parent);
+
+			// The root will always have a black parent, so this check covers both.
+			if (Parent == Root)
+			{
+				Nodes.SetColor(Parent, EColor::Black);
+				break;
+			}
+			else if (IsBlack(Parent))
+			{
+				break;
+			}
+
+			const FIntervalTreeNodeIndex GrandParent = Nodes.Parents[Parent];
+
+			UE_AUTORTFM_ASSUME((Parent != GrandParent) && (Current != GrandParent));
+
+			const bool bParentIsLeft = Nodes.Lefts[GrandParent] == Parent;
+
+			// The uncle is the other node of our parent.
+			const FIntervalTreeNodeIndex Uncle = bParentIsLeft ? Nodes.Rights[GrandParent] : Nodes.Lefts[GrandParent];
+
+			UE_AUTORTFM_ASSUME((GrandParent != Uncle) && (Parent != Uncle) && (Current != Uncle));
+
+			if (!IsBlack(Uncle))
+			{
+				AUTORTFM_ASSERT(IsBlack(GrandParent));
+				Nodes.SetColor(Parent, EColor::Black);
+				Nodes.SetColor(Uncle, EColor::Black);
+
+				if (GrandParent == Root)
+				{
+					break;
+				}
+				else
+				{
+					Nodes.SetColor(GrandParent, EColor::Red);
+				}
+
+				Current = GrandParent;
+				continue;
+			}
+
+			// Our uncle is black, so we need to swizzle around.
+			const bool bCurrentIsLeft = Nodes.Lefts[Parent] == Current;
+
+			if (bParentIsLeft)
+			{
+				if (!bCurrentIsLeft)
+				{
+					Nodes.Rights[Parent] = Nodes.Lefts[Current];
+
+					if (IntervalTreeNodeIndexNone != Nodes.Rights[Parent])
+					{
+						Nodes.Parents[Nodes.Rights[Parent]] = Parent;
+					}
+
+					Nodes.Parents[Parent] = Current;
+					Nodes.Parents[Current] = GrandParent;
+					Nodes.Lefts[Current] = Parent;
+					Nodes.Lefts[GrandParent] = Current;
+					std::swap(Parent, Current);
+				}
+
+				Nodes.Parents[Parent] = Nodes.Parents[GrandParent];
+				Nodes.Lefts[GrandParent] = Nodes.Rights[Parent];
+
+				if (IntervalTreeNodeIndexNone != Nodes.Lefts[GrandParent])
+				{
+					Nodes.Parents[Nodes.Lefts[GrandParent]] = GrandParent;
+				}
+
+				Nodes.Rights[Parent] = GrandParent;
+				Nodes.Parents[GrandParent] = Parent;
+
+				const EColor bGrandParentColor = Nodes.Color(GrandParent);
+				Nodes.SetColor(GrandParent, Nodes.Color(Parent));
+				Nodes.SetColor(Parent, bGrandParentColor);
+			}
+			else
+			{
+				if (bCurrentIsLeft)
+				{
+					Nodes.Lefts[Parent] = Nodes.Rights[Current];
+					
+					if (IntervalTreeNodeIndexNone != Nodes.Lefts[Parent])
+					{
+						Nodes.Parents[Nodes.Lefts[Parent]] = Parent;
+					}
+
+					Nodes.Parents[Parent] = Current;
+					Nodes.Parents[Current] = GrandParent;
+					Nodes.Rights[Current] = Parent;
+					Nodes.Rights[GrandParent] = Current;
+					std::swap(Parent, Current);
+				}
+
+				Nodes.Parents[Parent] = Nodes.Parents[GrandParent];
+				Nodes.Rights[GrandParent] = Nodes.Lefts[Parent];
+
+				if (IntervalTreeNodeIndexNone != Nodes.Rights[GrandParent])
+				{
+					Nodes.Parents[Nodes.Rights[GrandParent]] = GrandParent;
+				}
+
+				Nodes.Lefts[Parent] = GrandParent;
+				Nodes.Parents[GrandParent] = Parent;
+
+				const EColor bGrandParentColor = Nodes.Color(GrandParent);
+				Nodes.SetColor(GrandParent, Nodes.Color(Parent));
+				Nodes.SetColor(Parent, bGrandParentColor);
+			}
+
+			if (IntervalTreeNodeIndexNone == Nodes.Parents[Parent])
+			{
+				Root = Parent;
+			}
+			else
+			{
+				const FIntervalTreeNodeIndex GreatGrandParent = Nodes.Parents[Parent];
+
+				if (GrandParent == Nodes.Lefts[GreatGrandParent])
+				{
+					Nodes.Lefts[GreatGrandParent] = Parent;
+				}
+				else
+				{
+					Nodes.Rights[GreatGrandParent] = Parent;
+				}
+			}
+
+			break;
+		}
 
 		AssertStructureIsOk();
 
-        return true;
-    }
+		return true;
+	}
 
 	UE_AUTORTFM_FORCEINLINE void AssertStructureIsOk() const
 	{
@@ -415,11 +534,11 @@ private:
 		if constexpr (bExtraDebugging)
 		{
 			AUTORTFM_ASSERT(IntervalTreeNodeIndexNone != Index);
-			AUTORTFM_ASSERT(Index < NodeRanges.Num());
+			AUTORTFM_ASSERT(Index < Nodes.Size);
 
-			const FIntervalTreeNodeIndex Parent = NodeParents[Index];
-			const FIntervalTreeNodeIndex Left = NodeLefts[Index];
-			const FIntervalTreeNodeIndex Right = NodeRights[Index];
+			const FIntervalTreeNodeIndex Parent = Nodes.Parents[Index];
+			const FIntervalTreeNodeIndex Left = Nodes.Lefts[Index];
+			const FIntervalTreeNodeIndex Right = Nodes.Rights[Index];
 
 			if (IntervalTreeNodeIndexNone == Parent)
 			{
@@ -427,8 +546,8 @@ private:
 			}
 			else
 			{
-				AUTORTFM_ASSERT(NodeColors[Parent] || NodeColors[Index]);
-				AUTORTFM_ASSERT((NodeLefts[Parent] == Index) ^ (NodeRights[Parent] == Index));
+				AUTORTFM_ASSERT(Nodes.Color(Parent) == EColor::Black || Nodes.Color(Index) == EColor::Black);
+				AUTORTFM_ASSERT((Nodes.Lefts[Parent] == Index) ^ (Nodes.Rights[Parent] == Index));
 			}
 
 			AUTORTFM_ASSERT(Left != Index);
