[AutoRTFM] Add FIntervalTree::FIterator, replace existing traversal logic with it

We'll need FIntervalTree::Next() to fix FIntervalTree::Contains() for queries where the start query interval is before an overlapping inserted interval.

#Synced-CL 48944287

diff --git a/Engine/Source/Programs/AutoRTFMTests/Private/IntervalTreeTests.cpp b/Engine/Source/Programs/AutoRTFMTests/Private/IntervalTreeTests.cpp
--- a/Engine/Source/Programs/AutoRTFMTests/Private/IntervalTreeTests.cpp
+++ b/Engine/Source/Programs/AutoRTFMTests/Private/IntervalTreeTests.cpp
@@ -96,10 +96,10 @@ TEST_CASE("IntervalTree")
 	auto Validate = [&](AutoRTFM::FIntervalTree& IntervalTree, std::vector<FInterval> ExpectedIntervals)
 	{
 		std::vector<FInterval> GotIntervals;
-		IntervalTree.ForEach([&](uintptr_t Start, uintptr_t End)
+		for (AutoRTFM::FIntervalTree::FRange Range : IntervalTree)
 		{
-			GotIntervals.push_back({Start, End});
-		});
+			GotIntervals.push_back({Range.Start, Range.End});
+		}
 		REQUIRE(ExpectedIntervals.empty() == IntervalTree.IsEmpty());
 		for (size_t I = 0, N = std::min(ExpectedIntervals.size(), GotIntervals.size()); I < N; I++)
 		{
@@ -283,7 +283,6 @@ TEST_CASE("IntervalTree.Benchmarks")
 		}
 		return Intervals;
 	};
-	
 
 	BENCHMARK_ADVANCED("Insert")(Catch::Benchmark::Chronometer Meter)
 	{
@@ -302,7 +301,7 @@ TEST_CASE("IntervalTree.Benchmarks")
 	// otherwise a random set of addresses outside of those intervals.
 	auto PopulateIntervalTreeAndPickAddresses = [&](bool bUsed) -> std::vector<void*>
 	{
-		std::mt19937 Rand(4242);
+		std::mt19937 Rand(/* Seed */ 4242);
 		
 		const std::vector<FInterval> Intervals = bUsed ? PopulateIntervalTree().Used : PopulateIntervalTree().Free;
 
@@ -339,4 +338,22 @@ TEST_CASE("IntervalTree.Benchmarks")
 			}
 		});
 	};
+
+	BENCHMARK_ADVANCED("Merge")(Catch::Benchmark::Chronometer Meter)
+	{
+		PopulateIntervalTree();
+		AutoRTFM::FIntervalTree OtherIntervalTree;
+		{
+			const FTreeIntervals Intervals = BuildIntervals(/* Seed */ 1234, /* MaxUsed */ 10'000);
+			for (FInterval Interval : Intervals.Used)
+			{
+				OtherIntervalTree.Insert(Ptr(Interval.Start), Interval.Size());
+			}
+		}
+
+		Meter.measure([&]
+		{
+			IntervalTree.Merge(OtherIntervalTree);
+		});
+	};
 }
diff --git a/Engine/Source/Runtime/AutoRTFM/Private/IntervalTree.h b/Engine/Source/Runtime/AutoRTFM/Private/IntervalTree.h
--- a/Engine/Source/Runtime/AutoRTFM/Private/IntervalTree.h
+++ b/Engine/Source/Runtime/AutoRTFM/Private/IntervalTree.h
@@ -15,6 +15,20 @@ namespace AutoRTFM
 
 struct AUTORTFM_DISABLE_NO_SANITIZE FIntervalTree final
 {
+private:
+	using NodeIndex = uint32_t;
+	static constexpr NodeIndex InvalidNode = UINT32_MAX;
+
+public:
+	struct FRange final
+	{
+		FRange(const void* const Address, const size_t Size) :
+			Start(reinterpret_cast<uintptr_t>(Address)), End(reinterpret_cast<uintptr_t>(Address) + Size) {}
+
+		uintptr_t Start;
+		uintptr_t End;
+	};
+
 	FIntervalTree() = default;
 
 	FIntervalTree(const FIntervalTree&) = delete;
@@ -30,12 +44,12 @@ struct AUTORTFM_DISABLE_NO_SANITIZE FIntervalTree final
 	{
 		const FRange NewRange(Address, Size);
 
-		if (AUTORTFM_UNLIKELY(IntervalTreeNodeIndexNone == Root))
+		if (AUTORTFM_UNLIKELY(InvalidNode == Root))
 		{
 			return false;
 		}
 
-		FIntervalTreeNodeIndex Current = Root;
+		NodeIndex Current = Root;
 
 		do
 		{
@@ -49,21 +63,21 @@ struct AUTORTFM_DISABLE_NO_SANITIZE FIntervalTree final
 				return true;
 			}
 
-			const FIntervalTreeNodeIndex* Next = (NewRange.Start < Range.End) ? Nodes.Lefts : Nodes.Rights;
+			const NodeIndex* Next = (NewRange.Start < Range.End) ? Nodes.Lefts : Nodes.Rights;
 			Current = Next[Current];
-		} while (IntervalTreeNodeIndexNone != Current);
+		} while (InvalidNode != Current);
 
 		return false;
 	}
 
 	bool IsEmpty() const
 	{
-		return IntervalTreeNodeIndexNone == Root;
+		return InvalidNode == Root;
 	}
 
 	void Reset()
 	{
-		Root = IntervalTreeNodeIndexNone;
+		Root = InvalidNode;
 		Nodes.Reset();
 	}
 
@@ -74,58 +88,75 @@ struct AUTORTFM_DISABLE_NO_SANITIZE FIntervalTree final
 			return;
 		}
 
-		TStack<FIntervalTreeNodeIndex, 32, EContainerValidation::Disabled> ToProcess;
+		TStack<NodeIndex, 32, EContainerValidation::Disabled> ToProcess;
 
 		ToProcess.Push(Other.Root);
 
 		do
 		{
-			const FIntervalTreeNodeIndex Current = ToProcess.Back();
+			const NodeIndex Current = ToProcess.Back();
 			ToProcess.Pop();
 
 			FRange Range = Other.Nodes.Ranges[Current];
 
 			Insert(std::move(Range));
 
-			const FIntervalTreeNodeIndex Left = Other.Nodes.Lefts[Current];
-			const FIntervalTreeNodeIndex Right = Other.Nodes.Rights[Current];
+			const NodeIndex Left = Other.Nodes.Lefts[Current];
+			const NodeIndex Right = Other.Nodes.Rights[Current];
 
-			if (IntervalTreeNodeIndexNone != Left)
+			if (InvalidNode != Left)
 			{
 				ToProcess.Push(Left);
 			}
 
-			if (IntervalTreeNodeIndexNone != Right)
+			if (InvalidNode != Right)
 			{
 				ToProcess.Push(Right);
 			}
 		} while (!ToProcess.IsEmpty());
 	}
 
-	// Calls Callback with each interval in the tree, in order of address.
-	// Callback must have the signature: Callback(uintptr_t Start, uintptr_t End)
-	template<typename CALLBACK>
-	void ForEach(CALLBACK&& Callback) const
+	// A forward iterator of the FIntervalTree
+	struct FIterator
 	{
-		if (Root != IntervalTreeNodeIndexNone)
+		FIterator(const FIterator&) = default;
+		FIterator& operator=(const FIterator&) = default;
+ 
+		FRange operator*() const
 		{
-			ForEach(Root, Callback);
+			return Tree->Nodes.Ranges[Node];
 		}
-	}
 
-private:
-	struct FRange final
-	{
-		FRange(const void* const Address, const size_t Size) :
-			Start(reinterpret_cast<uintptr_t>(Address)), End(reinterpret_cast<uintptr_t>(Address) + Size) {}
+		FIterator& operator++()
+		{
+			Node = Tree->Next(Node);
+			return *this;
+		}
 
-		uintptr_t Start;
-		uintptr_t End;
+		bool operator == (const FIterator& Other) const = default;
+		bool operator != (const FIterator& Other) const = default;
+
+	private:
+		friend struct FIntervalTree;
+		FIterator(FIntervalTree const* Tree, NodeIndex Node) : Tree{Tree}, Node{Node} {}
+		FIntervalTree const* Tree = nullptr;
+		NodeIndex Node = InvalidNode;
 	};
 
-	using FIntervalTreeNodeIndex = uint32_t;
+	// Returns an iterator to the first interval in the tree
+	FIterator begin() const
+	{
+		return FIterator(this, Root == InvalidNode ? InvalidNode : LeftMostDescendant(Root));
+	}
+
+	// Returns an end iterator for the tree
+	FIterator end() const
+	{
+		return FIterator(this, InvalidNode);
+	}
 
-	static constexpr FIntervalTreeNodeIndex IntervalTreeNodeIndexNone = UINT32_MAX;
+private:
+	friend struct FIterator;
 
 	enum EColor : bool
 	{
@@ -136,9 +167,9 @@ private:
 	struct FNode final
 	{
 		FRange Range;
-		FIntervalTreeNodeIndex Left = IntervalTreeNodeIndexNone;
-		FIntervalTreeNodeIndex Right = IntervalTreeNodeIndexNone;
-		FIntervalTreeNodeIndex Parent = IntervalTreeNodeIndexNone;
+		NodeIndex Left = InvalidNode;
+		NodeIndex Right = InvalidNode;
+		NodeIndex Parent = InvalidNode;
 		EColor Color;
 	};
 
@@ -151,15 +182,15 @@ private:
 		using RightType = decltype(FNode::Right);
 		using ParentType = decltype(FNode::Parent);
 		static constexpr size_t NumColorWordBits = sizeof(ColorWord) * 8;
-		static constexpr FIntervalTreeNodeIndex InitialCapacity = 32;
+		static constexpr NodeIndex InitialCapacity = 32;
 
 		RangeType* Ranges = nullptr;
 		LeftType* Lefts = nullptr;
 		RightType* Rights = nullptr;
 		ParentType* Parents = nullptr;
 		ColorWord* ColorWords = nullptr;
-		FIntervalTreeNodeIndex Size = 0;
-		FIntervalTreeNodeIndex Capacity = 0;
+		NodeIndex Size = 0;
+		NodeIndex Capacity = 0;
 
 		FNodes() = default;
 		
@@ -184,7 +215,7 @@ private:
 
 		// Pushes a new node to the end of the list.
 		// Returns the node index.
-		FIntervalTreeNodeIndex Push(FNode Node)
+		NodeIndex Push(FNode Node)
 		{
 			if (Size >= Capacity)
 			{
@@ -197,10 +228,10 @@ private:
 			Parents[Size] = Node.Parent;
 			SetColor(Size, Node.Color);
 
-			return static_cast<FIntervalTreeNodeIndex>(Size++);
+			return static_cast<NodeIndex>(Size++);
 		}
 
-		void SetColor(FIntervalTreeNodeIndex Index, EColor Color)
+		void SetColor(NodeIndex Index, EColor Color)
 		{
 			ColorWord const Bit = static_cast<ColorWord>(1) << (Index & (NumColorWordBits - 1));
 			ColorWord& Word = ColorWords[Index / NumColorWordBits];
@@ -227,7 +258,7 @@ private:
 		FNodes(const FNodes&) = delete;
 		FNodes& operator=(const FNodes&) = delete;
 	
-		void Resize(FIntervalTreeNodeIndex NewCapacity)
+		void Resize(NodeIndex NewCapacity)
 		{
 			size_t TotalSize = 0;
 			size_t const RangesOffset = SubAllocate<RangeType>(TotalSize, NewCapacity);
@@ -275,45 +306,64 @@ private:
 		}
 	};
 
-	FIntervalTreeNodeIndex Root = IntervalTreeNodeIndexNone;
+	NodeIndex Root = InvalidNode;
 	FNodes Nodes;
 
 	static constexpr bool bExtraDebugging = false;
 
-	template<typename CALLBACK>
-	void ForEach(FIntervalTreeNodeIndex Index, CALLBACK&& Callback) const
+	// Returns the index of the next sequential node after Node.
+	NodeIndex Next(NodeIndex Node) const
 	{
-		if (FIntervalTreeNodeIndex Left = Nodes.Lefts[Index]; Left != IntervalTreeNodeIndexNone)
+		// If Node has a right child, then the next node is the left-most
+		// descendant of that node (or that node, if it has no descendants).
+		if (NodeIndex Right = Nodes.Rights[Node]; Right != InvalidNode)
 		{
-			ForEach(Left, Callback);
+			return LeftMostDescendant(Right);
 		}
-		Callback(Nodes.Ranges[Index].Start, Nodes.Ranges[Index].End);
-		if (FIntervalTreeNodeIndex Right = Nodes.Rights[Index]; Right != IntervalTreeNodeIndexNone)
+
+		// Otherwise the next node is the first ancestor where Node is under its left branch.
+		for (NodeIndex Parent = Nodes.Parents[Node]; Parent != InvalidNode; Parent = Nodes.Parents[Node])
+		{
+			if (Nodes.Lefts[Parent] == Node)
+			{
+				return Parent;
+			}
+			Node = Parent;
+		}
+
+		return InvalidNode;
+	}
+
+	// Returns the left-most leaf descendant of Node.
+	NodeIndex LeftMostDescendant(NodeIndex Node) const
+	{
+		for (NodeIndex Left = Nodes.Lefts[Node]; Left != InvalidNode; Left = Nodes.Lefts[Node])
 		{
-			ForEach(Right, Callback);
+			Node = Left;
 		}
+		return Node;
 	}
 
-	UE_AUTORTFM_FORCEINLINE FIntervalTreeNodeIndex AddNode(FRange&& NewRange)
+	UE_AUTORTFM_FORCEINLINE NodeIndex AddNode(FRange&& NewRange)
 	{
 		return Nodes.Push(FNode{.Range = NewRange, .Color = EColor::Black});
 	}
 
-	UE_AUTORTFM_FORCEINLINE FIntervalTreeNodeIndex AddNode(FRange&& NewRange, FIntervalTreeNodeIndex Parent)
+	UE_AUTORTFM_FORCEINLINE NodeIndex AddNode(FRange&& NewRange, NodeIndex Parent)
 	{
 		return Nodes.Push(FNode{.Range = NewRange, .Parent = Parent, .Color = EColor::Red});
 	}
 
 	UE_AUTORTFM_FORCENOINLINE bool Insert(FRange&& NewRange)
 	{
-		if (AUTORTFM_UNLIKELY(IntervalTreeNodeIndexNone == Root))
+		if (AUTORTFM_UNLIKELY(InvalidNode == Root))
 		{
 			AUTORTFM_ASSERT(Nodes.Size == 0);
 			Root = AddNode(std::move(NewRange));
 			return true;
 		}
 
-		FIntervalTreeNodeIndex Current = Root;
+		NodeIndex Current = Root;
 
 		for(;;)
 		{
@@ -334,9 +384,9 @@ private:
 					Nodes.Ranges[Current].Start = NewRange.Start;
 					return true;
 				}
-				else if (IntervalTreeNodeIndexNone == Nodes.Lefts[Current])
+				else if (InvalidNode == Nodes.Lefts[Current])
 				{
-					const FIntervalTreeNodeIndex Index = AddNode(std::move(NewRange), Current);
+					const NodeIndex Index = AddNode(std::move(NewRange), Current);
 					Nodes.Lefts[Current] = Index;
 					Current = Index;
 					break;
@@ -354,9 +404,9 @@ private:
 					Nodes.Ranges[Current].End = NewRange.End;
 					return true;
 				}
-				else if (IntervalTreeNodeIndexNone == Nodes.Rights[Current])
+				else if (InvalidNode == Nodes.Rights[Current])
 				{
-					const FIntervalTreeNodeIndex Index = AddNode(std::move(NewRange), Current);
+					const NodeIndex Index = AddNode(std::move(NewRange), Current);
 					Nodes.Rights[Current] = Index;
 					Current = Index;
 					break;
@@ -368,14 +418,14 @@ private:
 			AUTORTFM_ASSERT(Root != Current);
 		}
 
-		auto IsBlack = [this](FIntervalTreeNodeIndex Index)
+		auto IsBlack = [this](NodeIndex Index)
 		{
-			return (IntervalTreeNodeIndexNone == Index) || (EColor::Black == Nodes.Color(Index));
+			return (InvalidNode == Index) || (EColor::Black == Nodes.Color(Index));
 		};
 
 		for(;;)
 		{
-			FIntervalTreeNodeIndex Parent = Nodes.Parents[Current];
+			NodeIndex Parent = Nodes.Parents[Current];
 
 			UE_AUTORTFM_ASSUME(Current != Parent);
 
@@ -390,14 +440,14 @@ private:
 				break;
 			}
 
-			const FIntervalTreeNodeIndex GrandParent = Nodes.Parents[Parent];
+			const NodeIndex GrandParent = Nodes.Parents[Parent];
 
 			UE_AUTORTFM_ASSUME((Parent != GrandParent) && (Current != GrandParent));
 
 			const bool bParentIsLeft = Nodes.Lefts[GrandParent] == Parent;
 
 			// The uncle is the other node of our parent.
-			const FIntervalTreeNodeIndex Uncle = bParentIsLeft ? Nodes.Rights[GrandParent] : Nodes.Lefts[GrandParent];
+			const NodeIndex Uncle = bParentIsLeft ? Nodes.Rights[GrandParent] : Nodes.Lefts[GrandParent];
 
 			UE_AUTORTFM_ASSUME((GrandParent != Uncle) && (Parent != Uncle) && (Current != Uncle));
 
@@ -429,7 +479,7 @@ private:
 				{
 					Nodes.Rights[Parent] = Nodes.Lefts[Current];
 
-					if (IntervalTreeNodeIndexNone != Nodes.Rights[Parent])
+					if (InvalidNode != Nodes.Rights[Parent])
 					{
 						Nodes.Parents[Nodes.Rights[Parent]] = Parent;
 					}
@@ -444,7 +494,7 @@ private:
 				Nodes.Parents[Parent] = Nodes.Parents[GrandParent];
 				Nodes.Lefts[GrandParent] = Nodes.Rights[Parent];
 
-				if (IntervalTreeNodeIndexNone != Nodes.Lefts[GrandParent])
+				if (InvalidNode != Nodes.Lefts[GrandParent])
 				{
 					Nodes.Parents[Nodes.Lefts[GrandParent]] = GrandParent;
 				}
@@ -462,7 +512,7 @@ private:
 				{
 					Nodes.Lefts[Parent] = Nodes.Rights[Current];
 					
-					if (IntervalTreeNodeIndexNone != Nodes.Lefts[Parent])
+					if (InvalidNode != Nodes.Lefts[Parent])
 					{
 						Nodes.Parents[Nodes.Lefts[Parent]] = Parent;
 					}
@@ -477,7 +527,7 @@ private:
 				Nodes.Parents[Parent] = Nodes.Parents[GrandParent];
 				Nodes.Rights[GrandParent] = Nodes.Lefts[Parent];
 
-				if (IntervalTreeNodeIndexNone != Nodes.Rights[GrandParent])
+				if (InvalidNode != Nodes.Rights[GrandParent])
 				{
 					Nodes.Parents[Nodes.Rights[GrandParent]] = GrandParent;
 				}
@@ -490,13 +540,13 @@ private:
 				Nodes.SetColor(Parent, bGrandParentColor);
 			}
 
-			if (IntervalTreeNodeIndexNone == Nodes.Parents[Parent])
+			if (InvalidNode == Nodes.Parents[Parent])
 			{
 				Root = Parent;
 			}
 			else
 			{
-				const FIntervalTreeNodeIndex GreatGrandParent = Nodes.Parents[Parent];
+				const NodeIndex GreatGrandParent = Nodes.Parents[Parent];
 
 				if (GrandParent == Nodes.Lefts[GreatGrandParent])
 				{
@@ -520,27 +570,27 @@ private:
 	{
 		if constexpr (bExtraDebugging)
 		{
-			if (IntervalTreeNodeIndexNone != Root)
+			if (InvalidNode != Root)
 			{
 				AssertNodeIsOk(Root);
 			}
 		}
 	}
 
-	UE_AUTORTFM_FORCENOINLINE void AssertNodeIsOk(FIntervalTreeNodeIndex Index) const
+	UE_AUTORTFM_FORCENOINLINE void AssertNodeIsOk(NodeIndex Index) const
 	{
 		// We need to use recursion to check because we cannot have any
 		// allocations within the checker itself!
 		if constexpr (bExtraDebugging)
 		{
-			AUTORTFM_ASSERT(IntervalTreeNodeIndexNone != Index);
+			AUTORTFM_ASSERT(InvalidNode != Index);
 			AUTORTFM_ASSERT(Index < Nodes.Size);
 
-			const FIntervalTreeNodeIndex Parent = Nodes.Parents[Index];
-			const FIntervalTreeNodeIndex Left = Nodes.Lefts[Index];
-			const FIntervalTreeNodeIndex Right = Nodes.Rights[Index];
+			const NodeIndex Parent = Nodes.Parents[Index];
+			const NodeIndex Left = Nodes.Lefts[Index];
+			const NodeIndex Right = Nodes.Rights[Index];
 
-			if (IntervalTreeNodeIndexNone == Parent)
+			if (InvalidNode == Parent)
 			{
 				AUTORTFM_ASSERT(Root == Index);
 			}
@@ -553,12 +603,12 @@ private:
 			AUTORTFM_ASSERT(Left != Index);
 			AUTORTFM_ASSERT(Right != Index);
 
-			if (IntervalTreeNodeIndexNone != Left)
+			if (InvalidNode != Left)
 			{
 				AssertNodeIsOk(Left);
 			}
 
-			if (IntervalTreeNodeIndexNone != Right)
+			if (InvalidNode != Right)
 			{
 				AssertNodeIsOk(Right);
 			}
