[AutoRTFM] Add sanitizer mode to record closed write callstacks

This will be printed, along with the open write, if the sanitizer hits a mixed open / closed write, which is extremely helpful in diagnosing the bug.

This adds a data record to TIntervalTree. There's work to do to make interval merging work correctly. This is tracked by SOL-8824.

#Synced-CL 48989719

diff --git a/Engine/Source/Programs/AutoRTFMTests/Private/AutoRTFMTests.cpp b/Engine/Source/Programs/AutoRTFMTests/Private/AutoRTFMTests.cpp
--- a/Engine/Source/Programs/AutoRTFMTests/Private/AutoRTFMTests.cpp
+++ b/Engine/Source/Programs/AutoRTFMTests/Private/AutoRTFMTests.cpp
@@ -64,11 +64,16 @@ int RunTests(int ArgC, const char* ArgV[])
 	bool bNoRetry = false;
 	bool bRetryNestedToo = false;
 	bool bEnableBenchmarks = false;
+	AUTORTFM_SANITIZER_ONLY(bool bSanitizerRecordWriteCallstacks = false);
 
 	Session.cli(Session.cli()
 		| Catch::Clara::Opt(bNoRetry)["--no-retry"]
 		| Catch::Clara::Opt(bRetryNestedToo)["--retry-nested-too"]
-		| Catch::Clara::Opt(bEnableBenchmarks)["--enable-benchmarks"]);
+		| Catch::Clara::Opt(bEnableBenchmarks)["--enable-benchmarks"]
+#if AUTORTFM_SANITIZER
+		| Catch::Clara::Opt(bSanitizerRecordWriteCallstacks)["--sanitizer-record-write-callstacks"]
+#endif
+	);
 
 	{
 		const int Result = Session.applyCommandLine(ArgC, ArgV);
@@ -86,6 +91,8 @@ int RunTests(int ArgC, const char* ArgV[])
 		}
 	}
 
+	AUTORTFM_SANITIZER_ONLY(AutoRTFM::ForTheRuntime::SetSanitizerRecordClosedWriteCallstacks(bSanitizerRecordWriteCallstacks));
+
 	FString CommandLine = TEXT("-Multiprocess -csvNoProcessingThread -LogCmds=\"LogCsvProfiler off, LogStreaming off, LogUObjectGlobals off, LogPackageName off, LogAutoRTFM warning\" -AsyncLoadingThread");
 #if FORCE_USE_STATS
 	CommandLine += TEXT(" -LoadTimeStatsForCommandlet");
diff --git a/Engine/Source/Programs/AutoRTFMTests/Private/CallstackTreeTests.cpp b/Engine/Source/Programs/AutoRTFMTests/Private/CallstackTreeTests.cpp
new file mode 100644
--- /dev/null
+++ b/Engine/Source/Programs/AutoRTFMTests/Private/CallstackTreeTests.cpp
@@ -0,0 +1,87 @@
+// Copyright Epic Games, Inc. All Rights Reserved.
+
+#include "CallstackTree.h"
+#include "Catch2Includes.h"
+
+TEST_CASE("CallstackTree")
+{
+	using FHandle = AutoRTFM::FCallstackTree::FHandle;
+	static constexpr AutoRTFM::FCallstackTree::FHandle InvalidHandle = AutoRTFM::FCallstackTree::InvalidHandle;
+
+	SECTION("Empty")
+	{
+		AutoRTFM::FCallstackTree Tree;
+		REQUIRE(Tree.Add(0, nullptr) == InvalidHandle);
+		REQUIRE((Tree.Get(InvalidHandle).IsEmpty()));
+	}
+
+	SECTION("Shallow, Unique")
+	{
+		AutoRTFM::FCallstackTree Tree;
+		AutoRTFM::TStack<autortfm_callstack_frame, 1> A{1};
+		FHandle HandleA = Tree.Add(A.Num(), &A.Front());
+		REQUIRE(HandleA == 1);
+
+		AutoRTFM::TStack<autortfm_callstack_frame, 1> B{2};
+		FHandle HandleB = Tree.Add(B.Num(), &B.Front());
+		REQUIRE(HandleB == 2);
+
+		AutoRTFM::TStack<autortfm_callstack_frame, 1> C{3};
+		FHandle HandleC = Tree.Add(C.Num(), &C.Front());
+		REQUIRE(HandleC == 3);
+
+		REQUIRE((Tree.Get(HandleA) == A));
+		REQUIRE((Tree.Get(HandleB) == B));
+		REQUIRE((Tree.Get(HandleC) == C));
+	}
+
+	SECTION("Shallow, Duplicates")
+	{
+		AutoRTFM::FCallstackTree Tree;
+		AutoRTFM::TStack<autortfm_callstack_frame, 1> A{1};
+		FHandle HandleA = Tree.Add(A.Num(), &A.Front());
+		REQUIRE(HandleA == 1);
+
+		AutoRTFM::TStack<autortfm_callstack_frame, 1> B{2};
+		FHandle HandleB = Tree.Add(B.Num(), &B.Front());
+		REQUIRE(HandleB == 2);
+
+		AutoRTFM::TStack<autortfm_callstack_frame, 1> C{1};
+		FHandle HandleC = Tree.Add(C.Num(), &C.Front());
+		REQUIRE(HandleC == 1);
+
+		REQUIRE((Tree.Get(HandleA) == A));
+		REQUIRE((Tree.Get(HandleB) == B));
+		REQUIRE((Tree.Get(HandleC) == C));
+	}
+	
+	SECTION("Deep")
+	{
+		AutoRTFM::FCallstackTree Tree;
+		AutoRTFM::TStack<autortfm_callstack_frame, 5> A{1, 2, 3, 4, 5};
+		FHandle HandleA = Tree.Add(A.Num(), &A.Front());
+		REQUIRE(HandleA == 5);
+
+		AutoRTFM::TStack<autortfm_callstack_frame, 5> B{1, 2, 3, 4, 5};
+		FHandle HandleB = Tree.Add(B.Num(), &B.Front());
+		REQUIRE(HandleB == 5);
+
+		AutoRTFM::TStack<autortfm_callstack_frame, 5> C{1, 2, 3};
+		FHandle HandleC = Tree.Add(C.Num(), &C.Front());
+		REQUIRE(HandleC == 3);
+
+		AutoRTFM::TStack<autortfm_callstack_frame, 5> D{4, 5};
+		FHandle HandleD = Tree.Add(D.Num(), &D.Front());
+		REQUIRE(HandleD == 7);
+
+		AutoRTFM::TStack<autortfm_callstack_frame, 5> E{1, 2};
+		FHandle HandleE = Tree.Add(E.Num(), &E.Front());
+		REQUIRE(HandleE == 2);
+
+		REQUIRE((Tree.Get(HandleA) == A));
+		REQUIRE((Tree.Get(HandleB) == B));
+		REQUIRE((Tree.Get(HandleC) == C));
+		REQUIRE((Tree.Get(HandleD) == D));
+		REQUIRE((Tree.Get(HandleE) == E));
+	}
+}
diff --git a/Engine/Source/Programs/AutoRTFMTests/Private/IntervalTreeTests.cpp b/Engine/Source/Programs/AutoRTFMTests/Private/IntervalTreeTests.cpp
--- a/Engine/Source/Programs/AutoRTFMTests/Private/IntervalTreeTests.cpp
+++ b/Engine/Source/Programs/AutoRTFMTests/Private/IntervalTreeTests.cpp
@@ -1,20 +1,16 @@
 // Copyright Epic Games, Inc. All Rights Reserved.
 
+#include "AutoRTFM.h"
 #include "IntervalTree.h"
 #include "Catch2Includes.h"
 #include <cstdint>
 #include <limits>
 #include <random>
+#include <type_traits>
 
-struct FInterval final
-{
-	uintptr_t Start = 0;
-	uintptr_t End = 0;
-	uintptr_t Size() const { return End - Start; }
-	bool operator < (const FInterval& Other) const { return Start < Other.Start; }
-};
+using FInterval = AutoRTFM::FInterval;
 
-static constexpr FInterval MaxInterval{0, std::numeric_limits<uintptr_t>::max()};
+static constexpr FInterval MaxInterval{static_cast<uintptr_t>(0), std::numeric_limits<uintptr_t>::max()};
 
 struct FTreeIntervals
 {
@@ -86,41 +82,98 @@ static FTreeIntervals BuildIntervals(int Seed, size_t MaxUsed)
 	return FTreeIntervals{Free, Used};
 }
 
-TEST_CASE("IntervalTree")
+TEMPLATE_TEST_CASE("IntervalTree", "", void, int)
 {
+	AutoRTFM::UnreachableIfClosed();
+
+	using IntervalTreeType = AutoRTFM::TIntervalTree<TestType>;
+	struct FIntervalAndData
+	{
+		FInterval Interval;
+		int Data;
+	};
+	static constexpr bool bHasData = !std::is_same_v<TestType, void>;
+
 	// Convenience lambda to convert a uintptr_t to a void*
 	auto Ptr = [](uintptr_t Address) { return reinterpret_cast<void*>(Address); };
 
+	using EntryType = std::conditional_t<bHasData, typename IntervalTreeType::FIntervalAndData, FInterval>;
+	auto IntervalOf = [] AUTORTFM_DISABLE (EntryType Entry)
+	{
+		if constexpr (bHasData)
+		{
+			return Entry.Interval;
+		}
+		else
+		{
+			return Entry;
+		}
+	};
+	auto MakeEntry = [] AUTORTFM_DISABLE (uintptr_t Address, uintptr_t Size, int Data)
+	{
+		if constexpr (bHasData)
+		{
+			return EntryType{FInterval{Address, Size}, Data};
+		}
+		else
+		{
+			return FInterval{Address, Size};
+		}
+	};
+	auto Insert = [Ptr] AUTORTFM_DISABLE (IntervalTreeType& Tree, uintptr_t Address, uintptr_t Size, int Data)
+	{
+		if constexpr (bHasData)
+		{
+			Tree.Insert(Ptr(Address), Size, Data);
+		}
+		else
+		{
+			Tree.Insert(Ptr(Address), Size);
+		}
+	};
+	
 	// Checks that the interval tree has all of the expected intervals, and that Contains() returns
 	// the correct value for various intervals.
-	auto Validate = [&](AutoRTFM::FIntervalTree& IntervalTree, std::vector<FInterval> ExpectedIntervals)
+	auto Validate = [&] AUTORTFM_DISABLE (IntervalTreeType& IntervalTree, std::vector<FIntervalAndData> Expected)
 	{
-		std::vector<FInterval> GotIntervals;
-		for (AutoRTFM::FIntervalTree::FRange Range : IntervalTree)
+		std::vector<FIntervalAndData> Got;
+		for (auto Entry : IntervalTree)
 		{
-			GotIntervals.push_back({Range.Start, Range.End});
+			if constexpr (bHasData)
+			{
+				Got.push_back(FIntervalAndData{Entry.Interval, Entry.Data});
+			}
+			else
+			{
+				Got.push_back(FIntervalAndData{Entry, {}});
+			}
 		}
-		REQUIRE(ExpectedIntervals.empty() == IntervalTree.IsEmpty());
-		for (size_t I = 0, N = std::min(ExpectedIntervals.size(), GotIntervals.size()); I < N; I++)
+
+		REQUIRE(Expected.empty() == IntervalTree.IsEmpty());
+		for (size_t I = 0, N = std::min(Expected.size(), Got.size()); I < N; I++)
 		{
-			REQUIRE(ExpectedIntervals[I].Start == GotIntervals[I].Start);
-			REQUIRE(ExpectedIntervals[I].End == GotIntervals[I].End);
+			REQUIRE(Expected[I].Interval.Start == Got[I].Interval.Start);
+			REQUIRE(Expected[I].Interval.End == Got[I].Interval.End);
+			if constexpr (bHasData)
+			{
+				REQUIRE(Expected[I].Data == Got[I].Data);
+			}
 		}
-		REQUIRE(ExpectedIntervals.size() == GotIntervals.size());
-		for (size_t IntervalIndex = 0; IntervalIndex < ExpectedIntervals.size(); IntervalIndex++)
+		REQUIRE(Expected.size() == Got.size());
+		for (size_t Index = 0; Index < Expected.size(); Index++)
 		{
-			const FInterval Interval = ExpectedIntervals[IntervalIndex];
-			const bool bTouchingPrev = IntervalIndex > 0 && ExpectedIntervals[IntervalIndex - 1].End == Interval.Start;
-			const bool bTouchingNext = IntervalIndex < (ExpectedIntervals.size() - 1) && ExpectedIntervals[IntervalIndex + 1].Start == Interval.End;
-			REQUIRE(true == IntervalTree.Contains(Ptr(Interval.Start), 1));
-			REQUIRE(true == IntervalTree.Contains(Ptr(Interval.Start), Interval.End - Interval.Start));
-			REQUIRE(bTouchingPrev == IntervalTree.Contains(Ptr(Interval.Start - 1), 1));
-			REQUIRE(bTouchingNext == IntervalTree.Contains(Ptr(Interval.End), 1));
+			const FIntervalAndData Entry = Expected[Index];
+			const bool bTouchingPrev = Index > 0 && Expected[Index - 1].Interval.End == Entry.Interval.Start;
+			const bool bTouchingNext = Index < (Expected.size() - 1) && Expected[Index + 1].Interval.Start == Entry.Interval.End;
+			REQUIRE(true == IntervalTree.Contains(Ptr(Entry.Interval.Start), 1));
+			REQUIRE(true == IntervalTree.Contains(Ptr(Entry.Interval.Start), Entry.Interval.Size()));
+			REQUIRE(bTouchingPrev == IntervalTree.Contains(Ptr(Entry.Interval.Start - 1), 1));
+			REQUIRE(bTouchingNext == IntervalTree.Contains(Ptr(Entry.Interval.End), 1));
 		}
 	};
 
-	AutoRTFM::FIntervalTree IntervalTree;
-	
+	IntervalTreeType IntervalTree;
+
 	Validate(IntervalTree, {});
 
 	SECTION("Insert")
@@ -129,104 +182,118 @@ TEST_CASE("IntervalTree")
 		Validate(IntervalTree, {});
 
 		// New intervals
-		IntervalTree.Insert(Ptr(40), 10);
-		Validate(IntervalTree, {{40, 50}});
+		Insert(IntervalTree, 40, 10, 1);
+		Validate(IntervalTree, {{{40, 50}, 1}});
 
-		IntervalTree.Insert(Ptr(20), 10);
-		Validate(IntervalTree, {{20, 30}, {40, 50}});
+		Insert(IntervalTree, 20, 10, 2);
+		Validate(IntervalTree, {{{20, 30}, 2}, {{40, 50}, 1}});
 
-		IntervalTree.Insert(Ptr(60), 10);
-		Validate(IntervalTree, {{20, 30}, {40, 50}, {60, 70}});
+		Insert(IntervalTree, 60, 10, 3);
+		Validate(IntervalTree, {{{20, 30}, 2}, {{40, 50}, 1}, {{60, 70}, 3}});
+
+		Insert(IntervalTree, 80, 10, 4);
+		Validate(IntervalTree, {{{20, 30}, 2}, {{40, 50}, 1}, {{60, 70}, 3}, {{80, 90}, 4}});
 
-		IntervalTree.Insert(Ptr(80), 10);
-		Validate(IntervalTree, {{20, 30}, {40, 50}, {60, 70}, {80, 90}});
-		
 		// Re-inserted intervals
-		IntervalTree.Insert(Ptr(20), 10);
-		Validate(IntervalTree, {{20, 30}, {40, 50}, {60, 70}, {80, 90}});
+		// TODO(SOL-8824): This is incorrect for intervals with data.
+		// The existing interval should be updated with the new data.
+		Insert(IntervalTree, 20, 10, 5);
+		Validate(IntervalTree, {{{20, 30}, 2}, {{40, 50}, 1}, {{60, 70}, 3}, {{80, 90}, 4}});
 
-		IntervalTree.Insert(Ptr(40), 10);
-		Validate(IntervalTree, {{20, 30}, {40, 50}, {60, 70}, {80, 90}});
+		Insert(IntervalTree, 40, 10, 6);
+		Validate(IntervalTree, {{{20, 30}, 2}, {{40, 50}, 1}, {{60, 70}, 3}, {{80, 90}, 4}});
 		
-		IntervalTree.Insert(Ptr(60), 10);
-		Validate(IntervalTree, {{20, 30}, {40, 50}, {60, 70}, {80, 90}});
+		Insert(IntervalTree, 60, 10, 7);
+		Validate(IntervalTree, {{{20, 30}, 2}, {{40, 50}, 1}, {{60, 70}, 3}, {{80, 90}, 4}});
 		
-		IntervalTree.Insert(Ptr(80), 10);
-		Validate(IntervalTree, {{20, 30}, {40, 50}, {60, 70}, {80, 90}});
+		Insert(IntervalTree, 80, 10, 8);
+		Validate(IntervalTree, {{{20, 30}, 2}, {{40, 50}, 1}, {{60, 70}, 3}, {{80, 90}, 4}});
 
 		// Expanding intervals without overlap
-		IntervalTree.Insert(Ptr(15), 5);
-		Validate(IntervalTree, {{15, 30}, {40, 50}, {60, 70}, {80, 90}});
+		// TODO(SOL-8824): This is incorrect for intervals with data.
+		// The existing interval should be reduced to no-overlaps, and the new
+		// interval should be inserted.
+		Insert(IntervalTree, 15, 5, 9);
+		Validate(IntervalTree, {{{15, 30}, 2}, {{40, 50}, 1}, {{60, 70}, 3}, {{80, 90}, 4}});
 
-		IntervalTree.Insert(Ptr(30), 5);
-		Validate(IntervalTree, {{15, 35}, {40, 50}, {60, 70}, {80, 90}});
+		Insert(IntervalTree, 30, 5, 10);
+		Validate(IntervalTree, {{{15, 35}, 2}, {{40, 50}, 1}, {{60, 70}, 3}, {{80, 90}, 4}});
 
-		IntervalTree.Insert(Ptr(55), 5);
-		Validate(IntervalTree, {{15, 35}, {40, 50}, {55, 70}, {80, 90}});
+		Insert(IntervalTree, 55, 5, 11);
+		Validate(IntervalTree, {{{15, 35}, 2}, {{40, 50}, 1}, {{55, 70}, 3}, {{80, 90}, 4}});
 
-		IntervalTree.Insert(Ptr(70), 5);
-		Validate(IntervalTree, {{15, 35}, {40, 50}, {55, 75}, {80, 90}});
+		Insert(IntervalTree, 70, 5, 12);
+		Validate(IntervalTree, {{{15, 35}, 2}, {{40, 50}, 1}, {{55, 75}, 3}, {{80, 90}, 4}});
 
 		// Expanding intervals with overlap
-		IntervalTree.Insert(Ptr(35), 5);
-		Validate(IntervalTree, {{15, 35}, {35, 50}, {55, 75}, {80, 90}});
+		// TODO(SOL-8824): This is incorrect for intervals with data.
+		// The existing interval should be reduced to no-overlaps, and the new
+		// interval should be inserted.
+		Insert(IntervalTree, 35, 5, 13);
+		Validate(IntervalTree, {{{15, 35}, 2}, {{35, 50}, 1}, {{55, 75}, 3}, {{80, 90}, 4}});
 
-		IntervalTree.Insert(Ptr(50), 5);
-		Validate(IntervalTree, {{15, 35}, {35, 55}, {55, 75}, {80, 90}});
+		Insert(IntervalTree, 50, 5, 14);
+		Validate(IntervalTree, {{{15, 35}, 2}, {{35, 55}, 1}, {{55, 75}, 3}, {{80, 90}, 4}});
 
-		IntervalTree.Insert(Ptr(75), 5);
-		Validate(IntervalTree, {{15, 35}, {35, 55}, {55, 80}, {80, 90}});
+		Insert(IntervalTree, 75, 5, 15);
+		Validate(IntervalTree, {{{15, 35}, 2}, {{35, 55}, 1}, {{55, 80}, 3}, {{80, 90}, 4}});
 	}
 
 	SECTION("Merge")
 	{
-		IntervalTree.Insert(Ptr(20), 10);
-		IntervalTree.Insert(Ptr(40), 10);
-		IntervalTree.Insert(Ptr(60), 10);
-		IntervalTree.Insert(Ptr(80), 10);
-		Validate(IntervalTree, {{20, 30}, {40, 50}, {60, 70}, {80, 90}});
+		Insert(IntervalTree, 20, 10, 1);
+		Insert(IntervalTree, 40, 10, 2);
+		Insert(IntervalTree, 60, 10, 3);
+		Insert(IntervalTree, 80, 10, 4);
+		Validate(IntervalTree, {{{20, 30}, 1}, {{40, 50}, 2}, {{60, 70}, 3}, {{80, 90}, 4}});
 
 		SECTION("Intertwined")
 		{
-			AutoRTFM::FIntervalTree IntervalTreeB;
-			IntervalTreeB.Insert(Ptr(10), 5);
-			IntervalTreeB.Insert(Ptr(33), 6);
-			IntervalTreeB.Insert(Ptr(54), 2);
-			IntervalTreeB.Insert(Ptr(73), 3);
-			IntervalTreeB.Insert(Ptr(95), 2);
-			Validate(IntervalTreeB, {{10, 15}, {33, 39}, {54, 56}, {73, 76}, {95, 97}});
+			IntervalTreeType IntervalTreeB;
+			Insert(IntervalTreeB, 10, 5, 10);
+			Insert(IntervalTreeB, 33, 6, 11);
+			Insert(IntervalTreeB, 54, 2, 12);
+			Insert(IntervalTreeB, 73, 3, 13);
+			Insert(IntervalTreeB, 95, 2, 14);
+			Validate(IntervalTreeB, {{{10, 15}, 10}, {{33, 39}, 11}, {{54, 56}, 12}, {{73, 76}, 13}, {{95, 97}, 14}});
 			
 			IntervalTree.Merge(IntervalTreeB);
-			Validate(IntervalTree, {{10, 15}, {20, 30}, {33, 39}, {40, 50}, {54, 56}, {60, 70}, {73, 76}, {80, 90}, {95, 97}});
+			Validate(IntervalTree, {{{10, 15}, 10}, {{20, 30}, 1}, {{33, 39}, 11}, {{40, 50}, 2}, {{54, 56}, 12}, {{60, 70}, 3}, {{73, 76}, 13}, {{80, 90}, 4}, {{95, 97}, 14}});
 		}
 
 		SECTION("Expand")
 		{
-			AutoRTFM::FIntervalTree IntervalTreeB;
-			IntervalTreeB.Insert(Ptr(15), 5);
-			IntervalTreeB.Insert(Ptr(30), 3);
-			IntervalTreeB.Insert(Ptr(37), 3);
-			IntervalTreeB.Insert(Ptr(50), 2);
-			IntervalTreeB.Insert(Ptr(55), 5);
-			IntervalTreeB.Insert(Ptr(70), 4);
-			IntervalTreeB.Insert(Ptr(78), 2);
-			IntervalTreeB.Insert(Ptr(90), 3);
-			Validate(IntervalTreeB, {{15, 20}, {30, 33}, {37, 40}, {50, 52}, {55, 60}, {70, 74}, {78, 80}, {90, 93}});
+			IntervalTreeType IntervalTreeB;
+			Insert(IntervalTreeB, 15, 5, 10);
+			Insert(IntervalTreeB, 30, 3, 11);
+			Insert(IntervalTreeB, 37, 3, 12);
+			Insert(IntervalTreeB, 50, 2, 13);
+			Insert(IntervalTreeB, 55, 5, 14);
+			Insert(IntervalTreeB, 70, 4, 15);
+			Insert(IntervalTreeB, 78, 2, 16);
+			Insert(IntervalTreeB, 90, 3, 17);
+			Validate(IntervalTreeB, {{{15, 20}, 10}, {{30, 33}, 11}, {{37, 40}, 12}, {{50, 52}, 13}, {{55, 60}, 14}, {{70, 74}, 15}, {{78, 80}, 16}, {{90, 93}, 17}});
 			
+			// TODO(SOL-8824): This is incorrect for intervals with data.
+			// The existing interval should be reduced to no-overlaps, and the new
+			// interval should be inserted.
 			IntervalTree.Merge(IntervalTreeB);
-			Validate(IntervalTree, {{15, 33}, {37, 52}, {55, 74}, {78, 93}});
+			Validate(IntervalTree, {{{15, 33}, 1}, {{37, 52}, 2}, {{55, 74}, 3}, {{78, 93}, 4}});
 		}
 
 		SECTION("Overlap")
 		{
-			AutoRTFM::FIntervalTree IntervalTreeB;
-			IntervalTreeB.Insert(Ptr(30), 10);
-			IntervalTreeB.Insert(Ptr(50), 10);
-			IntervalTreeB.Insert(Ptr(70), 10);
-			Validate(IntervalTreeB, {{30, 40}, {50, 60}, {70, 80}});
-			
+			IntervalTreeType IntervalTreeB;
+			Insert(IntervalTreeB, 30, 10, 10);
+			Insert(IntervalTreeB, 50, 10, 11);
+			Insert(IntervalTreeB, 70, 10, 12);
+			Validate(IntervalTreeB, {{{30, 40}, 10}, {{50, 60}, 11}, {{70, 80}, 12}});
+
+			// TODO(SOL-8824): This is incorrect for intervals with data.
+			// The existing interval should be reduced to no-overlaps, and the new
+			// interval should be inserted.
 			IntervalTree.Merge(IntervalTreeB);
-			Validate(IntervalTree, {{20, 30}, {30, 60}, {60, 80}, {80, 90}});
+			Validate(IntervalTree, {{{20, 30}, 1}, {{30, 60}, 2}, {{60, 80}, 3}, {{80, 90}, 4}});
 		}
 	}
 
@@ -241,7 +308,7 @@ TEST_CASE("IntervalTree")
 
 			for (FInterval Interval : Intervals.Used)
 			{
-				IntervalTree.Insert(Ptr(Interval.Start), Interval.Size());
+				Insert(IntervalTree, Interval.Start, Interval.Size(), 0);
 			}
 
 			Intervals.SortAndValidate();
diff --git a/Engine/Source/Runtime/AutoRTFM/Private/API.cpp b/Engine/Source/Runtime/AutoRTFM/Private/API.cpp
--- a/Engine/Source/Runtime/AutoRTFM/Private/API.cpp
+++ b/Engine/Source/Runtime/AutoRTFM/Private/API.cpp
@@ -31,6 +31,7 @@ namespace
 	float GAutoRTFMEnabledProbability = 100.0f;
 
 	AUTORTFM_SANITIZER_ONLY(AutoRTFM::ESanitizerMode GSanitizerMode = AutoRTFM::ESanitizerMode::Error);
+	AUTORTFM_SANITIZER_ONLY(bool GSanitizerRecordClosedWriteCallstacks = false);
 
 	bool GAutoRTFMEnsureOnInternalAbort = true;
 
@@ -330,8 +331,6 @@ namespace AutoRTFM
 		{
 			if (GSanitizerMode != Mode)
 			{
-				AUTORTFM_LOG("AutoRTFM sanitizer has been %s",
-					Mode == ESanitizerMode::Disabled ? "*disabled*" : "*enabled*");
 				GSanitizerMode = Mode;
 				if (GExternAPI.OnSanitizerModeChanged)
 				{
@@ -339,6 +338,23 @@ namespace AutoRTFM
 				}
 			}
 		}
+
+		bool SanitizerRecordClosedWriteCallstacks()
+		{
+			return GSanitizerRecordClosedWriteCallstacks;
+		}
+
+		void SetSanitizerRecordClosedWriteCallstacks(bool bEnable)
+		{
+			if (GSanitizerRecordClosedWriteCallstacks != bEnable)
+			{
+				GSanitizerRecordClosedWriteCallstacks = bEnable;
+				if (GExternAPI.OnSanitizerRecordClosedWriteCallstacksChanged)
+				{
+					GExternAPI.OnSanitizerRecordClosedWriteCallstacksChanged();
+				}
+			}
+		}
 #endif // AUTORTFM_SANITIZER
 
 	}
diff --git a/Engine/Source/Runtime/AutoRTFM/Private/CallstackTree.h b/Engine/Source/Runtime/AutoRTFM/Private/CallstackTree.h
new file mode 100644
--- /dev/null
+++ b/Engine/Source/Runtime/AutoRTFM/Private/CallstackTree.h
@@ -0,0 +1,86 @@
+// Copyright Epic Games, Inc. All Rights Reserved.
+
+#pragma once
+
+#if (defined(__AUTORTFM) && __AUTORTFM)
+
+#include "AutoRTFM/CAPI.h"
+#include "AutoRTFM/Defines.h"
+#include "HashMap.h"
+#include "Stack.h"
+#include "Utils.h"
+
+#include <algorithm>
+#include <cstdint>
+
+namespace AutoRTFM
+{
+
+class AUTORTFM_DISABLE_NO_SANITIZE FCallstackTree final
+{
+public:
+	using FHandle = uint32_t;
+	static constexpr FHandle InvalidHandle = 0;
+	static constexpr FHandle MaxHandle = ~static_cast<FHandle>(0);
+
+	FHandle Add(size_t StackCount, autortfm_callstack_frame const* StackTrace)
+	{
+		FHandle Parent = InvalidHandle;
+		FHandle Handle = InvalidHandle;
+		for (size_t I = 0; I < StackCount; I++)
+		{
+			FEntry const Entry{Parent, StackTrace[I]};
+			if (FHandle& Existing = Lookup.FindOrAdd(Entry))
+			{
+				Handle = Existing;
+			}
+			else
+			{
+				Entries.Push(Entry);
+				AUTORTFM_ASSERT_DEBUG(Entries.Num() <= MaxHandle);
+				Existing = static_cast<FHandle>(Entries.Num());
+				Handle = Existing;
+			}
+
+			Parent = Handle;
+		}
+		return Handle;
+	}
+
+	TStack<autortfm_callstack_frame, 8> Get(FHandle Handle)
+	{
+		TStack<autortfm_callstack_frame, 8> StackTrace;
+		while (Handle != InvalidHandle)
+		{
+			FEntry const& Entry = Entries[Handle - 1];
+			StackTrace.Push(Entry.PC);
+			Handle = Entry.Parent;
+		}
+		std::reverse(StackTrace.begin(), StackTrace.end());
+		return StackTrace;
+	}
+
+private:
+	struct FEntry
+	{
+		FHandle Parent;
+		autortfm_callstack_frame PC;
+		
+		bool operator == (const FEntry&) const = default;
+	};
+
+	struct FHasher
+	{
+		size_t operator()(const FEntry& Entry) const
+		{
+			return (Entry.Parent * 31) ^ Entry.PC;
+		}
+	};
+
+	THashMap<FEntry, FHandle, FHasher> Lookup;
+	TStack<FEntry, 32> Entries;
+};
+
+} // namespace AutoRTFM
+
+#endif // (defined(__AUTORTFM) && __AUTORTFM)
diff --git a/Engine/Source/Runtime/AutoRTFM/Private/Context.cpp b/Engine/Source/Runtime/AutoRTFM/Private/Context.cpp
--- a/Engine/Source/Runtime/AutoRTFM/Private/Context.cpp
+++ b/Engine/Source/Runtime/AutoRTFM/Private/Context.cpp
@@ -642,6 +642,33 @@ void FContext::DumpState() const
 	AUTORTFM_VERBOSE("Context at %p", this);
 }
 
+
+#if AUTORTFM_SANITIZER
+FCallstackTree::FHandle FContext::GetCallstackHandle(size_t FramesToSkip)
+{
+	static constexpr size_t MaxCallstack = 64;
+
+	FramesToSkip += 3; // Top frame is caller.
+
+	autortfm_callstack_frame Callstack[MaxCallstack];
+	size_t NumCaptured = GExternAPI.CaptureCallstack(MaxCallstack, Callstack);
+	if (NumCaptured <= FramesToSkip)
+	{
+		return FCallstackTree::InvalidHandle;
+	}
+
+	return CallstackTree.Add(NumCaptured - FramesToSkip, Callstack + FramesToSkip);
+}
+
+void FContext::LogCallstack(FCallstackTree::FHandle Handle, autortfm_log_severity Severity)
+{
+	if (auto Callstack = CallstackTree.Get(Handle); !Callstack.IsEmpty())
+	{
+		GExternAPI.LogCallstack(Severity, Callstack.Num(), &Callstack.Front());
+	}
+}
+#endif
+
 } // namespace AutoRTFM
 
 #endif // defined(__AUTORTFM) && __AUTORTFM
diff --git a/Engine/Source/Runtime/AutoRTFM/Private/Context.h b/Engine/Source/Runtime/AutoRTFM/Private/Context.h
--- a/Engine/Source/Runtime/AutoRTFM/Private/Context.h
+++ b/Engine/Source/Runtime/AutoRTFM/Private/Context.h
@@ -14,6 +14,10 @@
 #include "Transaction.h"
 #include "Utils.h"
 
+#if AUTORTFM_SANITIZER
+#include "CallstackTree.h"
+#endif
+
 namespace AutoRTFM
 {
 
@@ -133,6 +137,11 @@ public:
 	// Returns the stack address limits for the caller thread.
 	static FStackRange GetThreadStackRange();
 
+#if AUTORTFM_SANITIZER
+	FCallstackTree::FHandle GetCallstackHandle(size_t FramesToSkip);
+	void LogCallstack(FCallstackTree::FHandle, autortfm_log_severity);
+#endif
+
 private:
 	UE_AUTORTFM_API static FContext* Instance;
 
@@ -167,6 +176,10 @@ private:
 	FTaskPool TaskPool;
 	FTransactionPool TransactionPool;
 	FCallNestPool CallNestPool;
+
+#if AUTORTFM_SANITIZER
+	FCallstackTree CallstackTree;
+#endif
 };
 
 } // namespace AutoRTFM
diff --git a/Engine/Source/Runtime/AutoRTFM/Private/ExternAPI.cpp b/Engine/Source/Runtime/AutoRTFM/Private/ExternAPI.cpp
--- a/Engine/Source/Runtime/AutoRTFM/Private/ExternAPI.cpp
+++ b/Engine/Source/Runtime/AutoRTFM/Private/ExternAPI.cpp
@@ -67,12 +67,17 @@ autortfm_extern_api GExternAPI
 	/* AllocateZeroed */ nullptr,
 	/* Free */ nullptr,
 	/* Log */ DefaultLog,
-	/* LogWithCallback */ nullptr,
+	/* LogWithCallstack */ nullptr,
 	/* EnsureFailure */ nullptr,
 	/* IsLogActive */ nullptr,
+	/* CaptureCallstack */ nullptr,
+	/* LogCallstack */ nullptr,
 	/* OnRuntimeEnabledChanged */ nullptr,
 	/* OnRetryTransactionsChanged */ nullptr,
-	/* OnSanitizerModeChanged */ AUTORTFM_SANITIZER_ONLY(nullptr,)
+#if AUTORTFM_SANITIZER
+	/* OnSanitizerModeChanged */ nullptr,
+	/* OnSanitizerRecordClosedWriteCallstacksChanged */ nullptr,
+#endif
 	/* RollbackWrite */ nullptr,
 };
 
diff --git a/Engine/Source/Runtime/AutoRTFM/Private/IntervalTree.h b/Engine/Source/Runtime/AutoRTFM/Private/IntervalTree.h
--- a/Engine/Source/Runtime/AutoRTFM/Private/IntervalTree.h
+++ b/Engine/Source/Runtime/AutoRTFM/Private/IntervalTree.h
@@ -8,66 +8,67 @@
 #include "ContainerValidation.h"
 #include "Stack.h"
 #include "Utils.h"
+
 #include <cstring>
+#include <type_traits>
+#include <variant>
 
 namespace AutoRTFM
 {
 
-struct AUTORTFM_DISABLE_NO_SANITIZE FIntervalTree final
+struct FInterval final
+{
+	constexpr FInterval(uintptr_t Start, uintptr_t End) :
+		Start{Start}, End{End} {}
+
+	FInterval(const void* const Address, const size_t Size) :
+		Start(reinterpret_cast<uintptr_t>(Address)), End(reinterpret_cast<uintptr_t>(Address) + Size) {}
+
+	size_t Size() const { return End - Start; }
+	bool operator < (const FInterval& Other) const { return Start < Other.Start; }
+
+	uintptr_t Start;
+	uintptr_t End;
+};
+
+template<typename DataTypeIn>
+struct AUTORTFM_DISABLE_NO_SANITIZE TIntervalTree final
 {
 private:
+	static constexpr bool bHasData = !std::is_same_v<DataTypeIn, void>;
 	using NodeIndex = uint32_t;
+	using DataType = std::conditional_t<bHasData, DataTypeIn, std::monostate>;
 	static constexpr NodeIndex InvalidNode = UINT32_MAX;
 
 public:
-	struct FRange final
-	{
-		FRange(const void* const Address, const size_t Size) :
-			Start(reinterpret_cast<uintptr_t>(Address)), End(reinterpret_cast<uintptr_t>(Address) + Size) {}
-
-		uintptr_t Start;
-		uintptr_t End;
-	};
-
-	FIntervalTree() = default;
+	TIntervalTree() = default;
 
-	FIntervalTree(const FIntervalTree&) = delete;
-	FIntervalTree& operator=(const FIntervalTree&) = delete;
+	TIntervalTree(const TIntervalTree&) = delete;
+	TIntervalTree& operator=(const TIntervalTree&) = delete;
 
-	UE_AUTORTFM_FORCENOINLINE bool Insert(const void* const Address, const size_t Size)
-	{        
-		FRange NewRange(Address, Size);
-		return Insert(std::move(NewRange));
+	bool Insert(const void* const Address, const size_t Size) requires (!bHasData)
+	{
+		return Insert(FInsertArgs{.Interval = FInterval{Address, Size}});
 	}
 
-	UE_AUTORTFM_FORCENOINLINE bool Contains(const void* const Address, const size_t Size) const
+	bool Insert(const void* const Address, const size_t Size, DataType&& Data) requires (bHasData)
 	{
-		const FRange NewRange(Address, Size);
-
-		if (AUTORTFM_UNLIKELY(InvalidNode == Root))
-		{
-			return false;
-		}
-
-		NodeIndex Current = Root;
-
-		do
-		{
-			const FRange Range = Nodes.Ranges[Current];
+		return Insert(FInsertArgs{.Interval = FInterval{Address, Size}, .Data = std::move(Data)});
+	}
 
-			// This check does not need to prove that NewRange is entirely
-			// enclosed within Range, because if any byte of NewRange was in the
-			// original Range then it **must** already have been new memory.
-			if ((NewRange.Start < Range.End) && (Range.Start < NewRange.End))
-			{
-				return true;
-			}
+	bool Insert(const void* const Address, const size_t Size, const DataType& Data) requires (bHasData)
+	{
+		return Insert(FInsertArgs{.Interval = FInterval{Address, Size}, .Data = Data});
+	}
 
-			const NodeIndex* Next = (NewRange.Start < Range.End) ? Nodes.Lefts : Nodes.Rights;
-			Current = Next[Current];
-		} while (InvalidNode != Current);
+	bool Contains(const void* const Address, const size_t Size) const
+	{
+		return Query(Address, Size);
+	}
 
-		return false;
+	const DataType* Get(const void* const Address, const size_t Size) const requires (bHasData)
+	{
+		return Query(Address, Size);
 	}
 
 	bool IsEmpty() const
@@ -81,7 +82,7 @@ public:
 		Nodes.Reset();
 	}
 
-	void Merge(const FIntervalTree& Other)
+	void Merge(const TIntervalTree& Other)
 	{
 		if (Other.IsEmpty())
 		{
@@ -97,12 +98,21 @@ public:
 			const NodeIndex Current = ToProcess.Back();
 			ToProcess.Pop();
 
-			FRange Range = Other.Nodes.Ranges[Current];
-
-			Insert(std::move(Range));
+			if constexpr (bHasData)
+			{
+				Insert(
+				{
+					.Interval = Other.Nodes.Interval(Current),
+					.Data = std::move(Other.Nodes.Data(Current)),
+				});
+			}
+			else
+			{
+				Insert({.Interval = Other.Nodes.Interval(Current)});
+			}
 
-			const NodeIndex Left = Other.Nodes.Lefts[Current];
-			const NodeIndex Right = Other.Nodes.Rights[Current];
+			const NodeIndex Left = Other.Nodes.Left(Current);
+			const NodeIndex Right = Other.Nodes.Right(Current);
 
 			if (InvalidNode != Left)
 			{
@@ -116,15 +126,26 @@ public:
 		} while (!ToProcess.IsEmpty());
 	}
 
-	// A forward iterator of the FIntervalTree
+	struct FIntervalAndData
+	{
+		FInterval Interval;
+		const DataType& Data;
+	};
+
+	// A forward iterator of the TIntervalTree
 	struct FIterator
 	{
 		FIterator(const FIterator&) = default;
 		FIterator& operator=(const FIterator&) = default;
  
-		FRange operator*() const
+		FInterval operator*() const requires(!bHasData)
+		{
+			return Tree->Nodes.Interval(Node);
+		}
+
+		FIntervalAndData operator*() const requires(bHasData)
 		{
-			return Tree->Nodes.Ranges[Node];
+			return {.Interval = Tree->Nodes.Interval(Node), .Data = Tree->Nodes.Data(Node)};
 		}
 
 		FIterator& operator++()
@@ -137,9 +158,9 @@ public:
 		bool operator != (const FIterator& Other) const = default;
 
 	private:
-		friend struct FIntervalTree;
-		FIterator(FIntervalTree const* Tree, NodeIndex Node) : Tree{Tree}, Node{Node} {}
-		FIntervalTree const* Tree = nullptr;
+		friend struct TIntervalTree;
+		FIterator(TIntervalTree const* Tree, NodeIndex Node) : Tree{Tree}, Node{Node} {}
+		TIntervalTree const* Tree = nullptr;
 		NodeIndex Node = InvalidNode;
 	};
 
@@ -164,32 +185,42 @@ private:
 		Black = true
 	};
 
+	using ColorWord = uint64_t;
+
 	struct FNode final
 	{
-		FRange Range;
+		FInterval Interval;
 		NodeIndex Left = InvalidNode;
 		NodeIndex Right = InvalidNode;
 		NodeIndex Parent = InvalidNode;
 		EColor Color;
+		DataType Data;
 	};
 
+	struct FNodePtrsNoData
+	{
+		FInterval* Interval = nullptr;
+		NodeIndex* Left = nullptr;
+		NodeIndex* Right = nullptr;
+		NodeIndex* Parent = nullptr;
+		ColorWord* Color = nullptr;
+	};
+
+	struct FNodePtrsWithData final : FNodePtrsNoData
+	{
+		DataType* Data = nullptr;
+	};
+
+	using FNodePtrs = std::conditional_t<bHasData, FNodePtrsWithData, FNodePtrsNoData>;
+
 	// A list of FNodes, stored as a structure-of-arrays.
 	struct FNodes final
 	{
-		using ColorWord = uint64_t;
-		using RangeType = decltype(FNode::Range);
-		using LeftType = decltype(FNode::Left);
-		using RightType = decltype(FNode::Right);
-		using ParentType = decltype(FNode::Parent);
 		static constexpr size_t NumColorWordBits = sizeof(ColorWord) * 8;
 		static constexpr NodeIndex InitialCapacity = 32;
 
-		RangeType* Ranges = nullptr;
-		LeftType* Lefts = nullptr;
-		RightType* Rights = nullptr;
-		ParentType* Parents = nullptr;
-		ColorWord* ColorWords = nullptr;
-		NodeIndex Size = 0;
+		FNodePtrs Ptrs;
+		NodeIndex Count = 0;
 		NodeIndex Capacity = 0;
 
 		FNodes() = default;
@@ -202,14 +233,19 @@ private:
 		// Clears the node list to an empty state.
 		void Reset()
 		{
-			AutoRTFM::Free(Memory);
-			Ranges = nullptr;
-			Lefts = nullptr;
-			Rights = nullptr;
-			Parents = nullptr;
-			ColorWords = nullptr;
-			Size = 0;
+			if constexpr (bHasData)
+			{
+				for (size_t I = 0; I < Count; I++)
+				{
+					Ptrs.Data[I].~DataType();
+				}
+			}
+
+			Ptrs = {};
+			Count = 0;
 			Capacity = 0;
+
+			AutoRTFM::Free(Memory);
 			Memory = nullptr;
 		}
 
@@ -217,24 +253,39 @@ private:
 		// Returns the node index.
 		NodeIndex Push(FNode Node)
 		{
-			if (Size >= Capacity)
+			if (Count >= Capacity)
 			{
 				Resize(std::max(Capacity * 2, InitialCapacity));
 			}
 
-			Ranges[Size] = Node.Range;
-			Lefts[Size] = Node.Left;
-			Rights[Size] = Node.Right;
-			Parents[Size] = Node.Parent;
-			SetColor(Size, Node.Color);
+			Ptrs.Interval[Count] = Node.Interval;
+			Ptrs.Left[Count] = Node.Left;
+			Ptrs.Right[Count] = Node.Right;
+			Ptrs.Parent[Count] = Node.Parent;
+			SetColor(Count, Node.Color);
+			if constexpr (bHasData)
+			{
+				new(&Ptrs.Data[Count]) DataType(std::move(Node.Data));
+			}
 
-			return static_cast<NodeIndex>(Size++);
+			return static_cast<NodeIndex>(Count++);
 		}
 
+		FInterval& Interval(NodeIndex Index) { return Ptrs.Interval[Index]; }
+		const FInterval& Interval(NodeIndex Index) const { return Ptrs.Interval[Index]; }
+		NodeIndex& Left(NodeIndex Index) { return Ptrs.Left[Index]; }
+		const NodeIndex& Left(NodeIndex Index) const { return Ptrs.Left[Index]; }
+		NodeIndex& Right(NodeIndex Index) { return Ptrs.Right[Index]; }
+		const NodeIndex& Right(NodeIndex Index) const { return Ptrs.Right[Index]; }
+		NodeIndex& Parent(NodeIndex Index) { return Ptrs.Parent[Index]; }
+		const NodeIndex& Parent(NodeIndex Index) const { return Ptrs.Parent[Index]; }
+		DataType& Data(NodeIndex Index) requires(bHasData) { return Ptrs.Data[Index]; }
+		const DataType& Data(NodeIndex Index) const requires(bHasData) { return Ptrs.Data[Index]; }
+
 		void SetColor(NodeIndex Index, EColor Color)
 		{
 			ColorWord const Bit = static_cast<ColorWord>(1) << (Index & (NumColorWordBits - 1));
-			ColorWord& Word = ColorWords[Index / NumColorWordBits];
+			ColorWord& Word = Ptrs.Color[Index / NumColorWordBits];
 			if (Color == EColor::Black)
 			{
 				Word |= Bit;
@@ -248,7 +299,7 @@ private:
 		EColor Color(size_t Index) const
 		{
 			ColorWord const Bit = static_cast<ColorWord>(1) << (Index & (NumColorWordBits - 1));
-			ColorWord& Word = ColorWords[Index / NumColorWordBits];
+			ColorWord& Word = Ptrs.Color[Index / NumColorWordBits];
 			return (Word & Bit) ? EColor::Black : EColor::Red;
 		}
 
@@ -261,33 +312,49 @@ private:
 		void Resize(NodeIndex NewCapacity)
 		{
 			size_t TotalSize = 0;
-			size_t const RangesOffset = SubAllocate<RangeType>(TotalSize, NewCapacity);
-			size_t const LeftsOffset = SubAllocate<LeftType>(TotalSize, NewCapacity);
-			size_t const RightsOffset = SubAllocate<RightType>(TotalSize, NewCapacity);
-			size_t const ParentsOffset = SubAllocate<ParentType>(TotalSize, NewCapacity);
-			size_t const ColorWordsOffset = SubAllocate<ColorWord>(TotalSize, NumColorWords(NewCapacity));
+			size_t const IntervalOffset = SubAllocate<FInterval>(TotalSize, NewCapacity);
+			size_t const LeftOffset = SubAllocate<NodeIndex>(TotalSize, NewCapacity);
+			size_t const RightOffset = SubAllocate<NodeIndex>(TotalSize, NewCapacity);
+			size_t const ParentOffset = SubAllocate<NodeIndex>(TotalSize, NewCapacity);
+			size_t const ColorOffset = SubAllocate<ColorWord>(TotalSize, NumColorWords(NewCapacity));
+			size_t const DataOffset = bHasData ? SubAllocate<DataType>(TotalSize, NewCapacity) : 0;
 			auto* const NewMemory = reinterpret_cast<std::byte*>(AutoRTFM::Allocate(TotalSize, alignof(FNode)));
-			auto* const NewRanges = std::launder(reinterpret_cast<RangeType*>(NewMemory + RangesOffset));
-			auto* const NewLefts = std::launder(reinterpret_cast<LeftType*>(NewMemory + LeftsOffset));
-			auto* const NewRights = std::launder(reinterpret_cast<RightType*>(NewMemory + RightsOffset));
-			auto* const NewParents = std::launder(reinterpret_cast<ParentType*>(NewMemory + ParentsOffset));
-			auto* const NewColorWords = std::launder(reinterpret_cast<ColorWord*>(NewMemory + ColorWordsOffset));
-			if (size_t const NumToCopy = std::min(Size, NewCapacity))
+			auto* const NewInterval = std::launder(reinterpret_cast<FInterval*>(NewMemory + IntervalOffset));
+			auto* const NewLeft = std::launder(reinterpret_cast<NodeIndex*>(NewMemory + LeftOffset));
+			auto* const NewRight = std::launder(reinterpret_cast<NodeIndex*>(NewMemory + RightOffset));
+			auto* const NewParent = std::launder(reinterpret_cast<NodeIndex*>(NewMemory + ParentOffset));
+			auto* const NewColor = std::launder(reinterpret_cast<ColorWord*>(NewMemory + ColorOffset));
+			auto* const NewData = bHasData ? std::launder(reinterpret_cast<DataType*>(NewMemory + DataOffset)) : nullptr;
+			if (size_t const NumToCopy = std::min(Count, NewCapacity))
 			{
-				memcpy(NewRanges, Ranges, sizeof(RangeType) * NumToCopy);
-				memcpy(NewLefts, Lefts, sizeof(LeftType) * NumToCopy);
-				memcpy(NewRights, Rights, sizeof(RightType) * NumToCopy);
-				memcpy(NewParents, Parents, sizeof(ParentType) * NumToCopy);
-				memcpy(NewColorWords, ColorWords, sizeof(ColorWord) * NumColorWords(NumToCopy));
+				memcpy(NewInterval, Ptrs.Interval, sizeof(FInterval) * NumToCopy);
+				memcpy(NewLeft, Ptrs.Left, sizeof(NodeIndex) * NumToCopy);
+				memcpy(NewRight, Ptrs.Right, sizeof(NodeIndex) * NumToCopy);
+				memcpy(NewParent, Ptrs.Parent, sizeof(NodeIndex) * NumToCopy);
+				memcpy(NewColor, Ptrs.Color, sizeof(ColorWord) * NumColorWords(NumToCopy));
+
+				if constexpr (bHasData)
+				{
+					for (size_t I = 0; I < NumToCopy; I++)
+					{
+						new(&NewData[I]) DataType(std::move(Ptrs.Data[I]));
+						Ptrs.Data[I].~DataType();
+					}
+				}
 			}
 
 			AutoRTFM::Free(Memory);
 
-			Ranges = NewRanges;
-			Lefts = NewLefts;
-			Rights = NewRights;
-			Parents = NewParents;
-			ColorWords = NewColorWords;
+			Ptrs.Interval = NewInterval;
+			Ptrs.Left = NewLeft;
+			Ptrs.Right = NewRight;
+			Ptrs.Parent = NewParent;
+			Ptrs.Color = NewColor;
+			if constexpr (bHasData)
+			{
+				Ptrs.Data = NewData;
+			}
+
 			Capacity = NewCapacity;
 			Memory = NewMemory;
 		}
@@ -316,15 +383,15 @@ private:
 	{
 		// If Node has a right child, then the next node is the left-most
 		// descendant of that node (or that node, if it has no descendants).
-		if (NodeIndex Right = Nodes.Rights[Node]; Right != InvalidNode)
+		if (NodeIndex Right = Nodes.Right(Node); Right != InvalidNode)
 		{
 			return LeftMostDescendant(Right);
 		}
 
 		// Otherwise the next node is the first ancestor where Node is under its left branch.
-		for (NodeIndex Parent = Nodes.Parents[Node]; Parent != InvalidNode; Parent = Nodes.Parents[Node])
+		for (NodeIndex Parent = Nodes.Parent(Node); Parent != InvalidNode; Parent = Nodes.Parent(Node))
 		{
-			if (Nodes.Lefts[Parent] == Node)
+			if (Nodes.Left(Parent) == Node)
 			{
 				return Parent;
 			}
@@ -337,29 +404,47 @@ private:
 	// Returns the left-most leaf descendant of Node.
 	NodeIndex LeftMostDescendant(NodeIndex Node) const
 	{
-		for (NodeIndex Left = Nodes.Lefts[Node]; Left != InvalidNode; Left = Nodes.Lefts[Node])
+		for (NodeIndex Left = Nodes.Left(Node); Left != InvalidNode; Left = Nodes.Left(Node))
 		{
 			Node = Left;
 		}
 		return Node;
 	}
 
-	UE_AUTORTFM_FORCEINLINE NodeIndex AddNode(FRange&& NewRange)
-	{
-		return Nodes.Push(FNode{.Range = NewRange, .Color = EColor::Black});
-	}
+	struct FInsertArgsNoData { FInterval Interval; };
+	struct FInsertArgsWithData { FInterval Interval; DataType Data; };
+	using FInsertArgs = std::conditional_t<bHasData, FInsertArgsWithData, FInsertArgsNoData>;
 
-	UE_AUTORTFM_FORCEINLINE NodeIndex AddNode(FRange&& NewRange, NodeIndex Parent)
+	UE_AUTORTFM_FORCEINLINE NodeIndex AddNode(FInsertArgs&& Args, NodeIndex Parent = InvalidNode)
 	{
-		return Nodes.Push(FNode{.Range = NewRange, .Parent = Parent, .Color = EColor::Red});
+		if constexpr (bHasData)
+		{
+			return Nodes.Push(
+			{
+				.Interval = Args.Interval,
+				.Parent = Parent,
+				.Color = (Parent == InvalidNode) ? EColor::Black : EColor::Red,
+				.Data = Args.Data,
+			});
+		}
+		else
+		{
+			return Nodes.Push(
+			{
+				.Interval = Args.Interval,
+				.Parent = Parent,
+				.Color = (Parent == InvalidNode) ? EColor::Black : EColor::Red,
+				.Data = {},
+			});
+		}
 	}
 
-	UE_AUTORTFM_FORCENOINLINE bool Insert(FRange&& NewRange)
+	UE_AUTORTFM_FORCENOINLINE bool Insert(FInsertArgs&& Args)
 	{
 		if (AUTORTFM_UNLIKELY(InvalidNode == Root))
 		{
-			AUTORTFM_ASSERT(Nodes.Size == 0);
-			Root = AddNode(std::move(NewRange));
+			AUTORTFM_ASSERT(Nodes.Count == 0);
+			Root = AddNode(std::move(Args));
 			return true;
 		}
 
@@ -367,52 +452,52 @@ private:
 
 		for(;;)
 		{
-			const FRange Range = Nodes.Ranges[Current];
+			const FInterval Interval = Nodes.Interval(Current);
 
-			if (AUTORTFM_UNLIKELY((NewRange.Start < Range.End) && (Range.Start < NewRange.End)))
+			if (AUTORTFM_UNLIKELY((Args.Interval.Start < Interval.End) && (Interval.Start < Args.Interval.End)))
 			{
 				return false;
 			}
 
-			if (NewRange.Start < Range.Start)
+			if (Args.Interval.Start < Interval.Start)
 			{
-				if (AUTORTFM_UNLIKELY(NewRange.End == Range.Start))
+				if (AUTORTFM_UNLIKELY(Args.Interval.End == Interval.Start))
 				{
-					AUTORTFM_ASSERT(NewRange.Start < Range.Start);
+					AUTORTFM_ASSERT(Args.Interval.Start < Interval.Start);
 
 					// We can just modify the existing node in place.
-					Nodes.Ranges[Current].Start = NewRange.Start;
+					Nodes.Interval(Current).Start = Args.Interval.Start;
 					return true;
 				}
-				else if (InvalidNode == Nodes.Lefts[Current])
+				else if (InvalidNode == Nodes.Left(Current))
 				{
-					const NodeIndex Index = AddNode(std::move(NewRange), Current);
-					Nodes.Lefts[Current] = Index;
+					const NodeIndex Index = AddNode(std::move(Args), Current);
+					Nodes.Left(Current) = Index;
 					Current = Index;
 					break;
 				}
 
-				Current = Nodes.Lefts[Current];
+				Current = Nodes.Left(Current);
 			}
 			else
 			{
-				if (AUTORTFM_UNLIKELY(NewRange.Start == Range.End))
+				if (AUTORTFM_UNLIKELY(Args.Interval.Start == Interval.End))
 				{
-					AUTORTFM_ASSERT(NewRange.End > Range.End);
+					AUTORTFM_ASSERT(Args.Interval.End > Interval.End);
 
 					// We can just modify the existing node in place.
-					Nodes.Ranges[Current].End = NewRange.End;
+					Nodes.Interval(Current).End = Args.Interval.End;
 					return true;
 				}
-				else if (InvalidNode == Nodes.Rights[Current])
+				else if (InvalidNode == Nodes.Right(Current))
 				{
-					const NodeIndex Index = AddNode(std::move(NewRange), Current);
-					Nodes.Rights[Current] = Index;
+					const NodeIndex Index = AddNode(std::move(Args), Current);
+					Nodes.Right(Current) = Index;
 					Current = Index;
 					break;
 				}
 
-				Current = Nodes.Rights[Current];
+				Current = Nodes.Right(Current);
 			}
 
 			AUTORTFM_ASSERT(Root != Current);
@@ -425,7 +510,7 @@ private:
 
 		for(;;)
 		{
-			NodeIndex Parent = Nodes.Parents[Current];
+			NodeIndex Parent = Nodes.Parent(Current);
 
 			UE_AUTORTFM_ASSUME(Current != Parent);
 
@@ -440,14 +525,14 @@ private:
 				break;
 			}
 
-			const NodeIndex GrandParent = Nodes.Parents[Parent];
+			const NodeIndex GrandParent = Nodes.Parent(Parent);
 
 			UE_AUTORTFM_ASSUME((Parent != GrandParent) && (Current != GrandParent));
 
-			const bool bParentIsLeft = Nodes.Lefts[GrandParent] == Parent;
+			const bool bParentIsLeft = Nodes.Left(GrandParent) == Parent;
 
 			// The uncle is the other node of our parent.
-			const NodeIndex Uncle = bParentIsLeft ? Nodes.Rights[GrandParent] : Nodes.Lefts[GrandParent];
+			const NodeIndex Uncle = bParentIsLeft ? Nodes.Right(GrandParent) : Nodes.Left(GrandParent);
 
 			UE_AUTORTFM_ASSUME((GrandParent != Uncle) && (Parent != Uncle) && (Current != Uncle));
 
@@ -471,36 +556,36 @@ private:
 			}
 
 			// Our uncle is black, so we need to swizzle around.
-			const bool bCurrentIsLeft = Nodes.Lefts[Parent] == Current;
+			const bool bCurrentIsLeft = Nodes.Left(Parent) == Current;
 
 			if (bParentIsLeft)
 			{
 				if (!bCurrentIsLeft)
 				{
-					Nodes.Rights[Parent] = Nodes.Lefts[Current];
+					Nodes.Right(Parent) = Nodes.Left(Current);
 
-					if (InvalidNode != Nodes.Rights[Parent])
+					if (InvalidNode != Nodes.Right(Parent))
 					{
-						Nodes.Parents[Nodes.Rights[Parent]] = Parent;
+						Nodes.Parent(Nodes.Right(Parent)) = Parent;
 					}
 
-					Nodes.Parents[Parent] = Current;
-					Nodes.Parents[Current] = GrandParent;
-					Nodes.Lefts[Current] = Parent;
-					Nodes.Lefts[GrandParent] = Current;
+					Nodes.Parent(Parent) = Current;
+					Nodes.Parent(Current) = GrandParent;
+					Nodes.Left(Current) = Parent;
+					Nodes.Left(GrandParent) = Current;
 					std::swap(Parent, Current);
 				}
 
-				Nodes.Parents[Parent] = Nodes.Parents[GrandParent];
-				Nodes.Lefts[GrandParent] = Nodes.Rights[Parent];
+				Nodes.Parent(Parent) = Nodes.Parent(GrandParent);
+				Nodes.Left(GrandParent) = Nodes.Right(Parent);
 
-				if (InvalidNode != Nodes.Lefts[GrandParent])
+				if (InvalidNode != Nodes.Left(GrandParent))
 				{
-					Nodes.Parents[Nodes.Lefts[GrandParent]] = GrandParent;
+					Nodes.Parent(Nodes.Left(GrandParent)) = GrandParent;
 				}
 
-				Nodes.Rights[Parent] = GrandParent;
-				Nodes.Parents[GrandParent] = Parent;
+				Nodes.Right(Parent) = GrandParent;
+				Nodes.Parent(GrandParent) = Parent;
 
 				const EColor bGrandParentColor = Nodes.Color(GrandParent);
 				Nodes.SetColor(GrandParent, Nodes.Color(Parent));
@@ -510,51 +595,51 @@ private:
 			{
 				if (bCurrentIsLeft)
 				{
-					Nodes.Lefts[Parent] = Nodes.Rights[Current];
+					Nodes.Left(Parent) = Nodes.Right(Current);
 					
-					if (InvalidNode != Nodes.Lefts[Parent])
+					if (InvalidNode != Nodes.Left(Parent))
 					{
-						Nodes.Parents[Nodes.Lefts[Parent]] = Parent;
+						Nodes.Parent(Nodes.Left(Parent)) = Parent;
 					}
 
-					Nodes.Parents[Parent] = Current;
-					Nodes.Parents[Current] = GrandParent;
-					Nodes.Rights[Current] = Parent;
-					Nodes.Rights[GrandParent] = Current;
+					Nodes.Parent(Parent) = Current;
+					Nodes.Parent(Current) = GrandParent;
+					Nodes.Right(Current) = Parent;
+					Nodes.Right(GrandParent) = Current;
 					std::swap(Parent, Current);
 				}
 
-				Nodes.Parents[Parent] = Nodes.Parents[GrandParent];
-				Nodes.Rights[GrandParent] = Nodes.Lefts[Parent];
+				Nodes.Parent(Parent) = Nodes.Parent(GrandParent);
+				Nodes.Right(GrandParent) = Nodes.Left(Parent);
 
-				if (InvalidNode != Nodes.Rights[GrandParent])
+				if (InvalidNode != Nodes.Right(GrandParent))
 				{
-					Nodes.Parents[Nodes.Rights[GrandParent]] = GrandParent;
+					Nodes.Parent(Nodes.Right(GrandParent)) = GrandParent;
 				}
 
-				Nodes.Lefts[Parent] = GrandParent;
-				Nodes.Parents[GrandParent] = Parent;
+				Nodes.Left(Parent) = GrandParent;
+				Nodes.Parent(GrandParent) = Parent;
 
 				const EColor bGrandParentColor = Nodes.Color(GrandParent);
 				Nodes.SetColor(GrandParent, Nodes.Color(Parent));
 				Nodes.SetColor(Parent, bGrandParentColor);
 			}
 
-			if (InvalidNode == Nodes.Parents[Parent])
+			if (InvalidNode == Nodes.Parent(Parent))
 			{
 				Root = Parent;
 			}
 			else
 			{
-				const NodeIndex GreatGrandParent = Nodes.Parents[Parent];
+				const NodeIndex GreatGrandParent = Nodes.Parent(Parent);
 
-				if (GrandParent == Nodes.Lefts[GreatGrandParent])
+				if (GrandParent == Nodes.Left(GreatGrandParent))
 				{
-					Nodes.Lefts[GreatGrandParent] = Parent;
+					Nodes.Left(GreatGrandParent) = Parent;
 				}
 				else
 				{
-					Nodes.Rights[GreatGrandParent] = Parent;
+					Nodes.Right(GreatGrandParent) = Parent;
 				}
 			}
 
@@ -566,6 +651,44 @@ private:
 		return true;
 	}
 
+	UE_AUTORTFM_FORCENOINLINE
+	std::conditional_t<bHasData, const DataType*, bool> Query(const void* const Address, const size_t Size) const
+	{
+		if (AUTORTFM_UNLIKELY(InvalidNode == Root))
+		{
+			return {};
+		}
+
+		const FInterval QueryInterval{Address, Size};
+
+		NodeIndex Current = Root;
+
+		do
+		{
+			const FInterval Interval = Nodes.Interval(Current);
+
+			// This check does not need to prove that QueryInterval is entirely
+			// enclosed within Interval, because if any byte of QueryInterval was in the
+			// original Interval then it **must** already have been new memory.
+			if ((QueryInterval.Start < Interval.End) && (Interval.Start < QueryInterval.End))
+			{
+				if constexpr (bHasData)
+				{
+					return &Nodes.Data(Current);
+				}
+				else
+				{
+					return true;
+				}
+			}
+
+			const NodeIndex* Next = (QueryInterval.Start < Interval.End) ? Nodes.Ptrs.Left : Nodes.Ptrs.Right;
+			Current = Next[Current];
+		} while (InvalidNode != Current);
+
+		return {};
+	}
+
 	UE_AUTORTFM_FORCEINLINE void AssertStructureIsOk() const
 	{
 		if constexpr (bExtraDebugging)
@@ -584,11 +707,11 @@ private:
 		if constexpr (bExtraDebugging)
 		{
 			AUTORTFM_ASSERT(InvalidNode != Index);
-			AUTORTFM_ASSERT(Index < Nodes.Size);
+			AUTORTFM_ASSERT(Index < Nodes.Count);
 
-			const NodeIndex Parent = Nodes.Parents[Index];
-			const NodeIndex Left = Nodes.Lefts[Index];
-			const NodeIndex Right = Nodes.Rights[Index];
+			const NodeIndex Parent = Nodes.Parent(Index);
+			const NodeIndex Left = Nodes.Left(Index);
+			const NodeIndex Right = Nodes.Right(Index);
 
 			if (InvalidNode == Parent)
 			{
@@ -597,7 +720,7 @@ private:
 			else
 			{
 				AUTORTFM_ASSERT(Nodes.Color(Parent) == EColor::Black || Nodes.Color(Index) == EColor::Black);
-				AUTORTFM_ASSERT((Nodes.Lefts[Parent] == Index) ^ (Nodes.Rights[Parent] == Index));
+				AUTORTFM_ASSERT((Nodes.Left(Parent) == Index) ^ (Nodes.Right(Parent) == Index));
 			}
 
 			AUTORTFM_ASSERT(Left != Index);
@@ -616,6 +739,9 @@ private:
 	}
 };
 
+// FIntervalTree is an TIntervalTree with no data
+using FIntervalTree = TIntervalTree<void>;
+
 } // namespace AutoRTFM
 
 #endif // (defined(__AUTORTFM) && __AUTORTFM)
diff --git a/Engine/Source/Runtime/AutoRTFM/Private/Stack.h b/Engine/Source/Runtime/AutoRTFM/Private/Stack.h
--- a/Engine/Source/Runtime/AutoRTFM/Private/Stack.h
+++ b/Engine/Source/Runtime/AutoRTFM/Private/Stack.h
@@ -9,6 +9,7 @@
 #include "Utils.h"
 
 #include <cstddef>
+#include <initializer_list>
 #include <type_traits>
 
 namespace AutoRTFM
@@ -38,6 +39,16 @@ public:
 	// Constructor.
 	TStack() = default;
 
+	// Constructor from initializer list.
+	TStack(std::initializer_list<T> List)
+	{
+		Reserve(List.size());
+		for (const T& Item : List)
+		{
+			Push(Item);
+		}
+	}
+
 	// Copy constructor.
 	TStack(const TStack& Other)
 	{
@@ -232,6 +243,24 @@ public:
 	T* end() { return Data + Count; }
 	const T* end() const { return Data + Count; }
 
+	// Equality operator
+	template<typename OtherT, size_t OtherInlineCapacity>
+	bool operator == (const TStack<OtherT, OtherInlineCapacity>& Other) const
+	{
+		if (Num() != Other.Num())
+		{
+			return false;
+		}
+		for (size_t I = 0, N = Num(); I < N; I++)
+		{
+			if (Data[I] != Other[I])
+			{
+				return false;
+			}
+		}
+		return true;
+	}
+
 private:
 	// Copies the data from Other to this stack.
 	// This stack must be empty before calling.
diff --git a/Engine/Source/Runtime/AutoRTFM/Private/Transaction.cpp b/Engine/Source/Runtime/AutoRTFM/Private/Transaction.cpp
--- a/Engine/Source/Runtime/AutoRTFM/Private/Transaction.cpp
+++ b/Engine/Source/Runtime/AutoRTFM/Private/Transaction.cpp
@@ -250,13 +250,6 @@ void FTransaction::CommitNested()
 		}
 
 		Parent->WriteLog.Push(Write);
-
-#if AUTORTFM_SANITIZER
-		if ((Write.Flags & EWriteFlags::NoSanitize) != EWriteFlags::NoSanitize)
-		{
-			Parent->ClosedWrites.Insert(Write.LogicalAddress, Write.Size);
-		}
-#endif
 	}
 
 	// Similarly we need to merge the hit-sets to the parent, skipping those
@@ -269,6 +262,17 @@ void FTransaction::CommitNested()
 		}
 	}
 
+#if AUTORTFM_SANITIZER
+	for (TClosedWriteTree:: FIntervalAndData Entry : ClosedWrites)
+	{
+		const void* Address = reinterpret_cast<void*>(Entry.Interval.Start);
+		if (!Parent->IsOnStack(Address))
+		{
+			Parent->ClosedWrites.Insert(Address, Entry.Interval.Size(), Entry.Data);
+		}
+	}
+#endif
+
 	// For all the deferred calls to `PopOnCommitHandler` that we couldn't
 	// process (because our transaction nest didn't `PushOnCommitHandler`)
 	// we need to move these to the parent now to handle them.
diff --git a/Engine/Source/Runtime/AutoRTFM/Private/Transaction.h b/Engine/Source/Runtime/AutoRTFM/Private/Transaction.h
--- a/Engine/Source/Runtime/AutoRTFM/Private/Transaction.h
+++ b/Engine/Source/Runtime/AutoRTFM/Private/Transaction.h
@@ -15,6 +15,10 @@
 #include "Utils.h"
 #include "WriteLog.h"
 
+#if AUTORTFM_SANITIZER
+#include "CallstackTree.h"
+#endif // AUTORTFM_SANITIZER
+
 #include <optional>
 
 namespace AutoRTFM
@@ -167,8 +171,9 @@ private:
 	FHitSet HitSet;
 	FIntervalTree NewMemoryTracker;
 #if AUTORTFM_SANITIZER
+	using TClosedWriteTree = TIntervalTree<FCallstackTree::FHandle>;
 	TStack<ESanitizerMode, 8> SanitizerModeStack;
-	FIntervalTree ClosedWrites;
+	TClosedWriteTree ClosedWrites;
 #endif
 	FWriteLog WriteLog;
 	TStatStorage<uint64_t> StatDepth = 1;
diff --git a/Engine/Source/Runtime/AutoRTFM/Private/TransactionInlines.h b/Engine/Source/Runtime/AutoRTFM/Private/TransactionInlines.h
--- a/Engine/Source/Runtime/AutoRTFM/Private/TransactionInlines.h
+++ b/Engine/Source/Runtime/AutoRTFM/Private/TransactionInlines.h
@@ -128,7 +128,11 @@ void FTransaction::RecordWrite(void* LogicalAddress, size_t Size, EWriteFlags Fl
 #if AUTORTFM_SANITIZER
 	if ((Flags & EWriteFlags::NoSanitize) != EWriteFlags::NoSanitize)
 	{
-		ClosedWrites.Insert(LogicalAddress, Size);
+		FCallstackTree::FHandle CallstackHandle = 
+			AutoRTFM::ForTheRuntime::SanitizerRecordClosedWriteCallstacks()
+				? Context->GetCallstackHandle(/* FramesToSkip */ 0)
+				: FCallstackTree::InvalidHandle;
+		ClosedWrites.Insert(LogicalAddress, Size, CallstackHandle);
 	}
 #endif
 }
@@ -190,7 +194,13 @@ template<unsigned SIZE> AUTORTFM_NO_ASAN UE_AUTORTFM_FORCENOINLINE void FTransac
 
 	WriteLog.PushSmall<SIZE>(static_cast<std::byte*>(LogicalAddress), static_cast<std::byte*>(LogicalAddress));
 
-	AUTORTFM_SANITIZER_ONLY(ClosedWrites.Insert(LogicalAddress, SIZE));
+#if AUTORTFM_SANITIZER
+	FCallstackTree::FHandle CallstackHandle =
+		AutoRTFM::ForTheRuntime::SanitizerRecordClosedWriteCallstacks()
+			? Context->GetCallstackHandle(/* FramesToSkip */ 0)
+			: FCallstackTree::InvalidHandle;
+	ClosedWrites.Insert(LogicalAddress, SIZE, CallstackHandle);
+#endif
 }
 
 UE_AUTORTFM_FORCEINLINE void FTransaction::DidAllocate(void* LogicalAddress, const size_t Size)
@@ -363,18 +373,27 @@ void FTransaction::SanitizerOpenWrite(void* Ptr, size_t Size)
 		// false if the queried interval starts outside of an interval added by
 		// FIntervalTree::Insert() but still overlaps an inserted interval.
 		// Fix this.
-		if (AUTORTFM_UNLIKELY(Transaction->ClosedWrites.Contains(Ptr, Size)))
+		if (const FCallstackTree::FHandle* Callstack = Transaction->ClosedWrites.Get(Ptr, Size); AUTORTFM_UNLIKELY(Callstack))
 		{
-			const char* Message =
-				"AutoRTFMSan: 0x%x bytes at %p was modified in the closed and now the open.\n"
-				"This may lead to memory corruption if the transaction is aborted.";
-			if (SanitizerMode == ESanitizerMode::Error)
+			const autortfm_log_severity Severity = (SanitizerMode == ESanitizerMode::Error) ? autortfm_log_error : autortfm_log_warn;
+
+			AUTORTFM_LOG_AT_SEVERITY(Severity, "AutoRTFMSan: 0x%zu bytes at %p was modified in the closed and now the open.", Size, Ptr);
+			AUTORTFM_LOG_AT_SEVERITY(Severity, "This may lead to memory corruption if the transaction is aborted.");
+			AUTORTFM_LOG_AT_SEVERITY(Severity, "");
+
+			if (*Callstack != FCallstackTree::InvalidHandle)
 			{
-				AUTORTFM_FATAL(Message, static_cast<int>(Size), Ptr);
+				AUTORTFM_LOG_AT_SEVERITY(Severity, "AutoRTFMSan: The last closed write to this address was at:");
+				Context->LogCallstack(*Callstack, Severity);
+				AUTORTFM_LOG_AT_SEVERITY(Severity, "");
 			}
-			else
+
+			::AutoRTFM::LogWithCallstack(Severity, "AutoRTFMSan: The open write was at:");
+			AUTORTFM_LOG_AT_SEVERITY(Severity, "");
+
+			if (SanitizerMode == ESanitizerMode::Error)
 			{
-				AUTORTFM_WARN_WITH_CALLSTACK(Message, static_cast<int>(Size), Ptr);
+				AUTORTFM_FATAL("<AutoRTFMSan fatal error>");
 			}
 		}
 	}
diff --git a/Engine/Source/Runtime/AutoRTFM/Private/Utils.h b/Engine/Source/Runtime/AutoRTFM/Private/Utils.h
--- a/Engine/Source/Runtime/AutoRTFM/Private/Utils.h
+++ b/Engine/Source/Runtime/AutoRTFM/Private/Utils.h
@@ -183,11 +183,13 @@ inline constexpr T Lerp(T Start, T End, T Fraction)
 
 #define AUTORTFM_REPORT_ERROR(Format, ...) ::AutoRTFM::ReportError(__FILE__, __LINE__, /* ProgramCounter */ nullptr, Format, ##__VA_ARGS__)
 
-#define AUTORTFM_VERBOSE(Format, ...) ::AutoRTFM::Log(__FILE__, __LINE__, /* ProgramCounter */ nullptr, autortfm_log_verbose, Format, ##__VA_ARGS__)
-#define AUTORTFM_LOG(Format, ...) ::AutoRTFM::Log(__FILE__, __LINE__, /* ProgramCounter */ nullptr, autortfm_log_info, Format, ##__VA_ARGS__)
-#define AUTORTFM_WARN(Format, ...) ::AutoRTFM::Log(__FILE__, __LINE__, /* ProgramCounter */ nullptr, autortfm_log_warn, Format, ##__VA_ARGS__)
-#define AUTORTFM_ERROR(Format, ...) ::AutoRTFM::Log(__FILE__, __LINE__, /* ProgramCounter */ nullptr, autortfm_log_error, Format, ##__VA_ARGS__)
-#define AUTORTFM_FATAL(Format, ...) ::AutoRTFM::Log(__FILE__, __LINE__, /* ProgramCounter */ nullptr, autortfm_log_fatal, Format, ##__VA_ARGS__)
+#define AUTORTFM_LOG_AT_SEVERITY(Severity, Format, ...) ::AutoRTFM::Log(__FILE__, __LINE__, /* ProgramCounter */ nullptr, Severity, Format, ##__VA_ARGS__)
+
+#define AUTORTFM_VERBOSE(Format, ...) AUTORTFM_LOG_AT_SEVERITY(autortfm_log_verbose, Format, ##__VA_ARGS__)
+#define AUTORTFM_LOG(Format, ...) AUTORTFM_LOG_AT_SEVERITY(autortfm_log_info, Format, ##__VA_ARGS__)
+#define AUTORTFM_WARN(Format, ...) AUTORTFM_LOG_AT_SEVERITY(autortfm_log_warn, Format, ##__VA_ARGS__)
+#define AUTORTFM_ERROR(Format, ...) AUTORTFM_LOG_AT_SEVERITY(autortfm_log_error, Format, ##__VA_ARGS__)
+#define AUTORTFM_FATAL(Format, ...) AUTORTFM_LOG_AT_SEVERITY(autortfm_log_fatal, Format, ##__VA_ARGS__)
 
 #define AUTORTFM_VERBOSE_V(Format, Args) ::AutoRTFM::LogV(__FILE__, __LINE__, /* ProgramCounter */ nullptr, autortfm_log_verbose, Format, Args)
 #define AUTORTFM_LOG_V(Format, Args) ::AutoRTFM::LogV(__FILE__, __LINE__, /* ProgramCounter */ nullptr, autortfm_log_info, Format, Args)
@@ -195,12 +197,6 @@ inline constexpr T Lerp(T Start, T End, T Fraction)
 #define AUTORTFM_ERROR_V(Format, Args) ::AutoRTFM::LogV(__FILE__, __LINE__, /* ProgramCounter */ nullptr, autortfm_log_error, Format, Args)
 #define AUTORTFM_FATAL_V(Format, Args) ::AutoRTFM::LogV(__FILE__, __LINE__, /* ProgramCounter */ nullptr, autortfm_log_fatal, Format, Args)
 
-#define AUTORTFM_VERBOSE_WITH_CALLSTACK(Format, ...) ::AutoRTFM::LogWithCallstack(autortfm_log_verbose, Format, ##__VA_ARGS__)
-#define AUTORTFM_LOG_WITH_CALLSTACK(Format, ...) ::AutoRTFM::LogWithCallstack(autortfm_log_info, Format, ##__VA_ARGS__)
-#define AUTORTFM_WARN_WITH_CALLSTACK(Format, ...) ::AutoRTFM::LogWithCallstack(autortfm_log_warn, Format, ##__VA_ARGS__)
-#define AUTORTFM_ERROR_WITH_CALLSTACK(Format, ...) ::AutoRTFM::LogWithCallstack(autortfm_log_error, Format, ##__VA_ARGS__)
-#define AUTORTFM_FATAL_WITH_CALLSTACK(Format, ...) ::AutoRTFM::LogWithCallstack(autortfm_log_fatal, Format, ##__VA_ARGS__)
-
 #define AUTORTFM_VERBOSE_IF(Condition, Format, ...) do { if (Condition) { AUTORTFM_VERBOSE(Format, ##__VA_ARGS__); } } while(false)
 #define AUTORTFM_LOG_IF(Condition, Format, ...) do { if (Condition) { AUTORTFM_LOG(Format, ##__VA_ARGS__); } } while(false)
 #define AUTORTFM_WARN_IF(Condition, Format, ...) do { if (Condition) { AUTORTFM_WARN(Format, ##__VA_ARGS__); } } while(false)
diff --git a/Engine/Source/Runtime/AutoRTFM/Public/AutoRTFM.h b/Engine/Source/Runtime/AutoRTFM/Public/AutoRTFM.h
--- a/Engine/Source/Runtime/AutoRTFM/Public/AutoRTFM.h
+++ b/Engine/Source/Runtime/AutoRTFM/Public/AutoRTFM.h
@@ -938,6 +938,12 @@ namespace ForTheRuntime
 
 	// Sets the AutoRTFM Sanitizer mode.
 	UE_AUTORTFM_API void SetSanitizerMode(ESanitizerMode Mode);
+
+	// Returns true if the AutoRTFM Sanitizer records the callstack of closed writes.
+	UE_AUTORTFM_API bool SanitizerRecordClosedWriteCallstacks();
+
+	// Sets whether the AutoRTFM Sanitizer records the callstack of closed writes.
+	UE_AUTORTFM_API void SetSanitizerRecordClosedWriteCallstacks(bool bEnable);
 #endif // AUTORTFM_SANITIZER
 
 	// Manually create a new transaction from open code and push it as a transaction nest.
diff --git a/Engine/Source/Runtime/AutoRTFM/Public/AutoRTFM/CAPI.h b/Engine/Source/Runtime/AutoRTFM/Public/AutoRTFM/CAPI.h
--- a/Engine/Source/Runtime/AutoRTFM/Public/AutoRTFM/CAPI.h
+++ b/Engine/Source/Runtime/AutoRTFM/Public/AutoRTFM/CAPI.h
@@ -95,6 +95,9 @@ typedef enum
 // An opaque unique identifier for a transaction.
 typedef uint64_t autortfm_transaction_id;
 
+// A single frame of a callstack capture.
+typedef unsigned long long autortfm_callstack_frame;
+
 // Function pointers used by AutoRTFM for heap allocations, etc.
 typedef struct
 {
@@ -133,6 +136,17 @@ typedef struct
 	// Must not be null.
 	bool (*IsLogActive)(autortfm_log_severity Severity);
 
+	// Performs a capture of the currently executing callstack, writing at most
+	// MaxCount frame addresses to StackOut.
+	// Returns the number of frame addresses actually written.
+	// Must not be null.
+	size_t (*CaptureCallstack)(size_t MaxCount, autortfm_callstack_frame* StackOut);
+
+	// Prints a callstack captured with CaptureCallstack to the log with the
+	// given severity.
+	// Must not be null.
+	void (*LogCallstack)(autortfm_log_severity Severity, size_t Count, autortfm_callstack_frame const* Stack);
+
 	// Optional callback to be informed when the value returned by
 	// ForTheRuntime::IsAutoRTFMRuntimeEnabled() changes.
 	// Can be null.
@@ -148,6 +162,11 @@ typedef struct
 	// ForTheRuntime::GetSanitizerMode() changes.
 	// Can be null.
 	void (*OnSanitizerModeChanged)();
+	
+	// Optional callback to be informed when the value returned by
+	// ForTheRuntime::SanitizerRecordClosedWriteCallstacks() changes.
+	// Can be null.
+	void (*OnSanitizerRecordClosedWriteCallstacksChanged)();
 #endif // AUTORTFM_SANITIZER
 
 	// Custom abort handler for writes that that were added with the
diff --git a/Engine/Source/Runtime/Core/Private/AutoRTFM/AutoRTFMUE.cpp b/Engine/Source/Runtime/Core/Private/AutoRTFM/AutoRTFMUE.cpp
--- a/Engine/Source/Runtime/Core/Private/AutoRTFM/AutoRTFMUE.cpp
+++ b/Engine/Source/Runtime/Core/Private/AutoRTFM/AutoRTFMUE.cpp
@@ -1,6 +1,8 @@
 // Copyright Epic Games, Inc. All Rights Reserved.
 
 #include "AutoRTFM/AutoRTFMUE.h"
+#include "HAL/Platform.h"
+#include <cstdint>
 
 #if UE_AUTORTFM
 
@@ -8,6 +10,7 @@
 #include "Containers/StringConv.h"
 #include "Containers/UnrealString.h"
 #include "GenericPlatform/GenericPlatformCrashContext.h"
+#include "GenericPlatform/GenericPlatformStackWalk.h"
 #include "HAL/IConsoleManager.h"
 #include "HAL/LowLevelMemTracker.h"
 #include "HAL/UnrealMemory.h"
@@ -25,6 +28,16 @@ static_assert(UE_AUTORTFM_ENABLED, "AutoRTFM/AutoRTFMUE.cpp requires the compile
 namespace
 {
 
+static bool IsDisable(const FString& String)
+{
+	return String == "disable" || String == "disabled" || String == "off" || String == "0";
+}
+
+static bool IsEnable(const FString& String)
+{
+	return String == "enable" || String == "enabled" || String == "on" || String == "1";
+}
+
 void OnAutoRTFMRuntimeEnabledChanged()
 {
 	const bool bEnabled = AutoRTFM::ForTheRuntime::IsAutoRTFMRuntimeEnabled();
@@ -65,6 +78,13 @@ void OnAutoRTFMSanitizerModeChanged()
 		return FGenericCrashContext::SetGameData(TEXT("AutoRTFMSanitizerMode"), TEXT("Enabled"));
 	}
 }
+
+void OnAutoRTFMSanitizerRecordClosedWriteCallstacksChanged()
+{
+	UE_LOG(LogAutoRTFM, Log, TEXT("AutoRTFM Sanitizer: %s recording of closed write callstacks"),
+		AutoRTFM::ForTheRuntime::SanitizerRecordClosedWriteCallstacks() ? TEXT("enabled") : TEXT("disabled"));
+}
+
 #endif // AUTORTFM_SANITIZER
 
 static FAutoConsoleVariable CVarAutoRTFMRuntimeEnabled(
@@ -159,7 +179,7 @@ static FAutoConsoleVariable CVarAutoRTFMSanitizerMode(
 	FConsoleVariableDelegate::CreateLambda([](IConsoleVariable* Variable)
 	{
 		FString Value = Variable->GetString().ToLower();
-		if (Value == "disable" || Value == "disabled" || Value == "off" || Value == "0")
+		if (IsDisable(Value))
 		{
 			AutoRTFM::ForTheRuntime::SetSanitizerMode(AutoRTFM::ESanitizerMode::Disabled);
 			return;
@@ -169,7 +189,7 @@ static FAutoConsoleVariable CVarAutoRTFMSanitizerMode(
 			AutoRTFM::ForTheRuntime::SetSanitizerMode(AutoRTFM::ESanitizerMode::Warn);
 			return;
 		}
-		if (Value == "error" || Value == "on" || Value == "1")
+		if (Value == "error" || IsEnable(Value))
 		{
 			AutoRTFM::ForTheRuntime::SetSanitizerMode(AutoRTFM::ESanitizerMode::Error);
 			return;
@@ -178,6 +198,28 @@ static FAutoConsoleVariable CVarAutoRTFMSanitizerMode(
 	}),
 	ECVF_Default
 );
+
+static FAutoConsoleVariable CVarAutoRTFMSanitizerRecordWriteCallstacks(
+	TEXT("AutoRTFMSanitizerRecordWriteCallstacks"),
+	AutoRTFM::ForTheRuntime::SanitizerRecordClosedWriteCallstacks(),
+	TEXT("Enables the recording of callstacks for closed writes, providing more information on AutoRTFMSan failures"),
+	FConsoleVariableDelegate::CreateLambda([](IConsoleVariable* Variable)
+	{
+		FString Value = Variable->GetString().ToLower();
+		if (IsDisable(Value))
+		{
+			AutoRTFM::ForTheRuntime::SetSanitizerRecordClosedWriteCallstacks(false);
+			return;
+		}
+		if (IsEnable(Value))
+		{
+			AutoRTFM::ForTheRuntime::SetSanitizerRecordClosedWriteCallstacks(true);
+			return;
+		}
+		UE_LOG(LogAutoRTFM, Fatal, TEXT("'AutoRTFMSanitizerRecordWriteCallstacks' CVar was set to '%s' which is not one of 'enable' or 'disable'!"), *Value);
+	}),
+	ECVF_Default
+);
 #endif
 
 static FAutoConsoleVariable CVarAutoRTFMEnabledProbability(
@@ -192,7 +234,7 @@ static FAutoConsoleVariable CVarAutoRTFMEnabledProbability(
 	ECVF_Default
 );
 
-bool IsSeverityActive(autortfm_log_severity Severity)
+bool AutoRTFMIsSeverityActive(autortfm_log_severity Severity)
 {
 	switch (Severity)
 	{
@@ -234,10 +276,10 @@ TStringConversion<TStringConvert<UTF8CHAR, TCHAR>, 128> FmtToTChar(const char* F
 }
 
 UE_AUTORTFM_ALWAYS_OPEN_NO_MEMORY_VALIDATION
-void AutoRTFMLog(const char* File, int Line, void* ProgramCounter, autortfm_log_severity Severity, const char* Format, va_list Args)
+void AutoRTFMLogV(const char* File, int Line, void* ProgramCounter, autortfm_log_severity Severity, const char* Format, va_list Args)
 {
 #if !NO_LOGGING
-	if (!IsSeverityActive(Severity))
+	if (!AutoRTFMIsSeverityActive(Severity))
 	{
 		return;
 	}
@@ -267,41 +309,67 @@ void AutoRTFMLog(const char* File, int Line, void* ProgramCounter, autortfm_log_
 			UE::Logging::Private::BasicLog<ELogVerbosity::Error>(Record, &LogAutoRTFM, Message.Get());
 			break;
 		case autortfm_log_fatal:
-			FDebug::DumpStackTraceToLog(TEXT("AutoRTFM backtrace"), ELogVerbosity::Error);
 			UE::Logging::Private::BasicFatalLogWithProgramCounter(Record, &LogAutoRTFM, ProgramCounter, Message.Get());
 			break;
 	}
 #endif
 }
 
+void AutoRTFMLogf(const char* File, int Line, void* ProgramCounter, autortfm_log_severity Severity, const char* Format, ...)
+{
+	va_list Args;
+	va_start(Args, Format);
+	AutoRTFMLogV(File, Line, ProgramCounter, Severity, Format, Args);
+	va_end(Args);
+}
+
+size_t AutoRTFMCaptureCallstack(size_t MaxCount, autortfm_callstack_frame* StackOut)
+{
+	return FPlatformStackWalk::CaptureStackBackTrace(StackOut, static_cast<uint32>(MaxCount));
+}
+
+void AutoRTFMLogCallstack(autortfm_log_severity Severity, size_t Count, autortfm_callstack_frame const* Stack)
+{
+#if !NO_LOGGING
+	for (size_t CurrentDepth = 0; CurrentDepth < Count; CurrentDepth++)
+	{
+		static constexpr size_t StringSize = 1024;
+		ANSICHAR String[StringSize] {'\0'};
+		FPlatformStackWalk::ProgramCounterToHumanReadableString(CurrentDepth, Stack[CurrentDepth], String, StringSize);
+		AutoRTFMLogf(/* File */ nullptr, /* Line */ 0, /* ProgramCounter */ nullptr, Severity, "%s", String);
+	}
+#endif
+}
+
 UE_AUTORTFM_ALWAYS_OPEN_NO_MEMORY_VALIDATION
 void AutoRTFMLogWithCallstack(void* ProgramCounter, autortfm_log_severity Severity, const char* Format, va_list Args)
 {
 #if !NO_LOGGING
-	if (!IsSeverityActive(Severity))
+	if (!AutoRTFMIsSeverityActive(Severity))
 	{
 		return;
 	}
 
-	auto Message = FmtToTChar(Format, Args);
+	AutoRTFMLogV(/* File */ nullptr, /* Line */ 0, ProgramCounter, Severity, Format, Args);
+	constexpr uint32 MaxBacktrace = 512;
+	uint64 Backtrace[MaxBacktrace];
+	uint32 Count = FPlatformStackWalk::CaptureStackBackTrace(Backtrace, MaxBacktrace);
+	uint32 NumFramesToSkip = 2;
+	if (ProgramCounter)
+	{
+		for (uint32 I = 0; I < Count; I++)
+		{
+			if (Backtrace[I] == static_cast<uint64>(reinterpret_cast<uintptr_t>(ProgramCounter)))
+			{
+				NumFramesToSkip = I;
+				break;
+			}
+		}
+	}
 
-	switch (Severity)
+	if (Count > NumFramesToSkip)
 	{
-		case autortfm_log_verbose:
-			FDebug::DumpStackTraceToLog(Message.Get(), ELogVerbosity::Verbose);
-			break;
-		case autortfm_log_info:
-			FDebug::DumpStackTraceToLog(Message.Get(), ELogVerbosity::Display);
-			break;
-		case autortfm_log_warn:
-			FDebug::DumpStackTraceToLog(Message.Get(), ELogVerbosity::Warning);
-			break;
-		case autortfm_log_error:
-			FDebug::DumpStackTraceToLog(Message.Get(), ELogVerbosity::Error);
-			break;
-		case autortfm_log_fatal:
-			FDebug::DumpStackTraceToLog(Message.Get(), ELogVerbosity::Fatal);
-			break;
+		AutoRTFMLogCallstack(Severity, Count - NumFramesToSkip, Backtrace + NumFramesToSkip);
 	}
 #endif
 }
@@ -355,13 +423,18 @@ CORE_API void AutoRTFM::InitializeForUE(RollbackWriteFn RollbackWrite /* = nullp
 		.Reallocate = AutoRTFMReallocate,
 		.AllocateZeroed = AutoRTFMAllocateZeroed,
 		.Free = AutoRTFMFree,
-		.Log = AutoRTFMLog,
+		.Log = AutoRTFMLogV,
 		.LogWithCallstack = AutoRTFMLogWithCallstack,
 		.EnsureFailure = AutoRTFMEnsureFailure,
-		.IsLogActive = IsSeverityActive,
+		.IsLogActive = AutoRTFMIsSeverityActive,
+		.CaptureCallstack = AutoRTFMCaptureCallstack,
+		.LogCallstack = AutoRTFMLogCallstack,
 		.OnRuntimeEnabledChanged = OnAutoRTFMRuntimeEnabledChanged,
 		.OnRetryTransactionsChanged = OnAutoRTFMRetryTransactionsChanged,
-		AUTORTFM_SANITIZER_ONLY(.OnSanitizerModeChanged = OnAutoRTFMSanitizerModeChanged,)
+#if AUTORTFM_SANITIZER
+		.OnSanitizerModeChanged = OnAutoRTFMSanitizerModeChanged,
+		.OnSanitizerRecordClosedWriteCallstacksChanged = OnAutoRTFMSanitizerRecordClosedWriteCallstacksChanged,
+#endif
 		.RollbackWrite = RollbackWrite,
 	};
 	
@@ -371,6 +444,10 @@ CORE_API void AutoRTFM::InitializeForUE(RollbackWriteFn RollbackWrite /* = nullp
 	// and crash context data is set.
 	OnAutoRTFMRuntimeEnabledChanged();
 	OnAutoRTFMRetryTransactionsChanged();
+#if AUTORTFM_SANITIZER
+	OnAutoRTFMSanitizerModeChanged();
+	OnAutoRTFMSanitizerRecordClosedWriteCallstacksChanged();
+#endif
 }
 
 #endif // UE_AUTORTFM
