[AutoRTFM] Optimize FCallstackTree for memory usage

The Lookup map held duplicate data to the Entries vector. Use a std::unordered_set with a custom hasher / equality type to remove the need to duplicate FEntry.

#Synced-CL 48989719

diff --git a/Engine/Source/Runtime/AutoRTFM/Private/CallstackTree.h b/Engine/Source/Runtime/AutoRTFM/Private/CallstackTree.h
--- a/Engine/Source/Runtime/AutoRTFM/Private/CallstackTree.h
+++ b/Engine/Source/Runtime/AutoRTFM/Private/CallstackTree.h
@@ -6,12 +6,13 @@
 
 #include "AutoRTFM/CAPI.h"
 #include "AutoRTFM/Defines.h"
-#include "HashMap.h"
 #include "Stack.h"
+#include "StlAllocator.h"
 #include "Utils.h"
 
 #include <algorithm>
 #include <cstdint>
+#include <unordered_set>
 
 namespace AutoRTFM
 {
@@ -30,16 +31,16 @@ public:
 		for (size_t I = 0; I < StackCount; I++)
 		{
 			FEntry const Entry{Parent, StackTrace[I]};
-			if (FHandle& Existing = Lookup.FindOrAdd(Entry))
+			if (auto It = Lookup.find(Entry); It != Lookup.end())
 			{
-				Handle = Existing;
+				Handle = *It;
 			}
 			else
 			{
 				Entries.Push(Entry);
 				AUTORTFM_ASSERT_DEBUG(Entries.Num() <= MaxHandle);
-				Existing = static_cast<FHandle>(Entries.Num());
-				Handle = Existing;
+				Handle = static_cast<FHandle>(Entries.Num());
+				Lookup.emplace_hint(It, Handle);
 			}
 
 			Parent = Handle;
@@ -69,16 +70,48 @@ private:
 		bool operator == (const FEntry&) const = default;
 	};
 
+	using FEntryList = TStack<FEntry, 32>;
+
 	struct FHasher
 	{
+		using is_transparent = void;
+		const FEntryList& Entries;
 		size_t operator()(const FEntry& Entry) const
 		{
 			return (Entry.Parent * 31) ^ Entry.PC;
 		}
+		size_t operator()(FHandle Handle) const
+		{
+			const FEntry& Entry = Entries[Handle - 1];
+			return operator()(Entry);
+		}
+	};
+
+	struct FEquality
+	{
+		using is_transparent = void;
+		const FEntryList& Entries;
+		size_t operator()(const FEntry& Entry, FHandle Handle) const
+		{
+			return Entries[Handle - 1] == Entry;
+		}
+		size_t operator()(FHandle Handle, const FEntry& Entry) const
+		{
+			return operator()(Entry, Handle);
+		}
+		size_t operator()(FHandle A, FHandle B) const
+		{
+			return A == B;
+		}
 	};
 
-	THashMap<FEntry, FHandle, FHasher> Lookup;
-	TStack<FEntry, 32> Entries;
+	FEntryList Entries;
+	std::unordered_set<FHandle, FHasher, FEquality, StlAllocator<FHandle>> Lookup
+	{
+		/* bucket_count */ 32,
+		/* hash */ FHasher{Entries},
+		/* equal */ FEquality{Entries},
+	};
 };
 
 } // namespace AutoRTFM
