[AutoRTFM] Move sanitizer APIs to new AutoRTFM::Sanitizer namespace, runtime-disable the sanitizer for AutoRTFM::Assign()

The operator=() in AutoRTFM::Assign() may not be inlined, in which case disabling the sanitizer instrumentation on AutoRTFM::Assign() may not be sufficient to avoid sanitizer false-positives.
Use a new AUTORTFM_DISABLE_SANITIZER_SCOPE() macro to runtime-disable the sanitizer for the scope of the assignment operator call.

Also drop the EWriteFlags::NoSanitize flag on the call to autortfm_record_open_write_with_flags() in AutoRTFM::Assign().
This was only required to avoid a false-positive with the old memory validator, which has now been expunged.

#Synced-CL 48989719

diff --git a/Engine/Source/Programs/AutoRTFMTests/Private/API.cpp b/Engine/Source/Programs/AutoRTFMTests/Private/API.cpp
--- a/Engine/Source/Programs/AutoRTFMTests/Private/API.cpp
+++ b/Engine/Source/Programs/AutoRTFMTests/Private/API.cpp
@@ -1276,7 +1276,7 @@ TEST_CASE("API.OnAbortMacro_NoAbort")
 
 TEST_CASE("API.OnPreAbortMacro_NoAbort")
 {
-	AUTORTFM_SCOPED_DISABLE_SANITIZER();
+	AUTORTFM_SCOPED_SANITIZER_MODE_DISABLE();
 
 	int Value = 123;
 
@@ -1314,7 +1314,7 @@ TEST_CASE("API.OnAbortMacro_WithAbort")
 
 TEST_CASE("API.OnPreAbortMacro_WithAbort")
 {
-	AUTORTFM_SCOPED_DISABLE_SANITIZER();
+	AUTORTFM_SCOPED_SANITIZER_MODE_DISABLE();
 
 	int Value = 123;
 
diff --git a/Engine/Source/Programs/AutoRTFMTests/Private/AutoRTFMTestUtils.h b/Engine/Source/Programs/AutoRTFMTests/Private/AutoRTFMTestUtils.h
--- a/Engine/Source/Programs/AutoRTFMTests/Private/AutoRTFMTestUtils.h
+++ b/Engine/Source/Programs/AutoRTFMTests/Private/AutoRTFMTestUtils.h
@@ -115,24 +115,24 @@ private:
 #if AUTORTFM_SANITIZER
 // A helper class that temporarily changes the AutoRTFM sanitizer mode for the
 // lifetime of the object, restoring the original level on destruction.
-class ScopedSanitizerMode
+class FScopedSanitizerMode
 {
 public:
-	ScopedSanitizerMode(AutoRTFM::ESanitizerMode NewMode)
-		: PrevMode(AutoRTFM::ForTheRuntime::GetSanitizerMode())
+	FScopedSanitizerMode(AutoRTFM::Sanitizer::EMode NewMode)
+		: PrevMode(AutoRTFM::Sanitizer::GetMode())
 	{
-		AutoRTFM::ForTheRuntime::SetSanitizerMode(NewMode);
+		AutoRTFM::Sanitizer::SetMode(NewMode);
 	}
 
-	~ScopedSanitizerMode()
+	~FScopedSanitizerMode()
 	{
-		AutoRTFM::ForTheRuntime::SetSanitizerMode(PrevMode);
+		AutoRTFM::Sanitizer::SetMode(PrevMode);
 	}
 private:
-	ScopedSanitizerMode(const ScopedSanitizerMode&) = delete;
-	ScopedSanitizerMode& operator = (const ScopedSanitizerMode&) = delete;
+	FScopedSanitizerMode(const FScopedSanitizerMode&) = delete;
+	FScopedSanitizerMode& operator = (const FScopedSanitizerMode&) = delete;
 
-	const AutoRTFM::ESanitizerMode PrevMode;
+	const AutoRTFM::Sanitizer::EMode PrevMode;
 };
 #endif // AUTORTFM_SANITIZER
 
@@ -141,14 +141,14 @@ private:
 		{AutoRTFM::ForTheRuntime::NoRetry}
 
 #if AUTORTFM_SANITIZER
-#define AUTORTFM_SCOPED_DISABLE_SANITIZER() \
-	AutoRTFMTestUtils::ScopedSanitizerMode ScopedSanitizerModeWarning \
-		{AutoRTFM::ESanitizerMode::Disabled}
+#define AUTORTFM_SCOPED_SANITIZER_MODE_DISABLE() \
+	AutoRTFMTestUtils::FScopedSanitizerMode ScopedSanitizerModeWarning \
+		{AutoRTFM::Sanitizer::EMode::Disabled}
 #define AUTORTFM_SCOPED_SANITIZER_MODE_WARN() \
-	AutoRTFMTestUtils::ScopedSanitizerMode ScopedSanitizerModeWarning \
-		{AutoRTFM::ESanitizerMode::Warn}
+	AutoRTFMTestUtils::FScopedSanitizerMode ScopedSanitizerModeWarning \
+		{AutoRTFM::Sanitizer::EMode::Warn}
 #else
-#define AUTORTFM_SCOPED_DISABLE_SANITIZER() static_assert(true) /* require semicolon */
+#define AUTORTFM_SCOPED_SANITIZER_MODE_DISABLE() static_assert(true) /* require semicolon */
 #endif
 
 static constexpr const std::string_view kMemoryModifiedInOpenWarning =
diff --git a/Engine/Source/Programs/AutoRTFMTests/Private/AutoRTFMTests.cpp b/Engine/Source/Programs/AutoRTFMTests/Private/AutoRTFMTests.cpp
--- a/Engine/Source/Programs/AutoRTFMTests/Private/AutoRTFMTests.cpp
+++ b/Engine/Source/Programs/AutoRTFMTests/Private/AutoRTFMTests.cpp
@@ -91,7 +91,7 @@ int RunTests(int ArgC, const char* ArgV[])
 		}
 	}
 
-	AUTORTFM_SANITIZER_ONLY(AutoRTFM::ForTheRuntime::SetSanitizerRecordClosedWriteCallstacks(bSanitizerRecordWriteCallstacks));
+	AUTORTFM_SANITIZER_ONLY(AutoRTFM::Sanitizer::SetRecordClosedWriteCallstacks(bSanitizerRecordWriteCallstacks));
 
 	FString CommandLine = TEXT("-Multiprocess -csvNoProcessingThread -LogCmds=\"LogCsvProfiler off, LogStreaming off, LogUObjectGlobals off, LogPackageName off, LogAutoRTFM warning\" -AsyncLoadingThread");
 #if FORCE_USE_STATS
diff --git a/Engine/Source/Programs/AutoRTFMTests/Private/FakeVM.cpp b/Engine/Source/Programs/AutoRTFMTests/Private/FakeVM.cpp
--- a/Engine/Source/Programs/AutoRTFMTests/Private/FakeVM.cpp
+++ b/Engine/Source/Programs/AutoRTFMTests/Private/FakeVM.cpp
@@ -61,7 +61,7 @@ AutoRTFM::ETransactionResult FakeVM(const FOpcode* Opcodes, const int Length, in
 {
 	// SOL-6743: These tests perform writes within open and closed scopes to the same memory
 	// locations, which will cause AutoRTFM sanitizer errors.
-	AUTORTFM_SCOPED_DISABLE_SANITIZER();
+	AUTORTFM_SCOPED_SANITIZER_MODE_DISABLE();
 
 	return AutoRTFM::Transact([Opcodes, Length, &Value]
 	{
diff --git a/Engine/Source/Programs/AutoRTFMTests/Private/Open.cpp b/Engine/Source/Programs/AutoRTFMTests/Private/Open.cpp
--- a/Engine/Source/Programs/AutoRTFMTests/Private/Open.cpp
+++ b/Engine/Source/Programs/AutoRTFMTests/Private/Open.cpp
@@ -435,7 +435,6 @@ TEST_CASE("Open.Collision")
 
             REQUIRE(I == 42);
             REQUIRE(J == 24);
-            REQUIRE(WarningContext.GetWarnings().empty());
         }
 
         SECTION("Transact(Open(RecordOpenWrite()), Open(Assign))")
@@ -453,7 +452,6 @@ TEST_CASE("Open.Collision")
             });
 
             REQUIRE(I == 42);
-            REQUIRE(WarningContext.GetWarnings().empty());
         }
 
         SECTION("Transact(Transact(Open(RecordOpenWrite())), Open(Assign))")
@@ -475,7 +473,6 @@ TEST_CASE("Open.Collision")
             });
 
             REQUIRE(I == 42);
-            REQUIRE(WarningContext.GetWarnings().empty());
         }
 
         SECTION("Transact(Transact(Open(RecordOpenWrite())), Open(RecordOpenWrite()), Open(Assign))")
@@ -500,7 +497,6 @@ TEST_CASE("Open.Collision")
             });
 
             REQUIRE(I == 24);
-            REQUIRE(WarningContext.GetWarnings().empty());
         }
 
         SECTION("Transact(Open(RecordOpenWrite()), Open(Assign)) <large>")
@@ -526,7 +522,6 @@ TEST_CASE("Open.Collision")
                 REQUIRE(V == 42);
             }
 
-            REQUIRE(WarningContext.GetWarnings().empty());
         }
 
         SECTION("Transact(Open(RecordOpenWrite()), Open(RecordOpenWrite()), Open(Assign))")
@@ -547,7 +542,6 @@ TEST_CASE("Open.Collision")
                 });
             });
 
-            REQUIRE(WarningContext.GetWarnings().empty());
         }
 
         SECTION("Transact(Open(RecordOpenWrite()), Open(RecordOpenWrite()), Open(Assign)) <large>")
@@ -573,7 +567,6 @@ TEST_CASE("Open.Collision")
                 });
             });
 
-            REQUIRE(WarningContext.GetWarnings().empty());
         }
 
         SECTION("Transact(Open(Assign), Assign)")
@@ -585,7 +578,6 @@ TEST_CASE("Open.Collision")
             });
 
             REQUIRE(I == 42);
-            REQUIRE(WarningContext.GetWarnings().empty());
         }
 
         SECTION("Transact(Assign, Open(Transact(Assign)))")
@@ -603,7 +595,6 @@ TEST_CASE("Open.Collision")
             });
 
             REQUIRE(I == 42);
-            REQUIRE(WarningContext.GetWarnings().empty());
         }
 
         SECTION("Transact(Open(Transact(Assign)), Assign)")
@@ -621,7 +612,6 @@ TEST_CASE("Open.Collision")
             });
 
             REQUIRE(I == 42);
-            REQUIRE(WarningContext.GetWarnings().empty());
         }
     }
 
@@ -641,6 +631,36 @@ TEST_CASE("Open.Collision")
         });
     }
 
+    SECTION("Transact(Assign, AUTORTFM_DISABLE_SANITIZER_SCOPE(Open(Assign)))")
+    {
+        AutoRTFM::Transact([&]
+        {
+            I = 42;
+
+            AUTORTFM_DISABLE_SANITIZER_SCOPE();
+
+            AutoRTFM::Open([&]
+            {
+                I = 24; // Sanitizer error silenced by AUTORTFM_DISABLE_SANITIZER_SCOPE()
+            });
+        });
+    }
+
+    SECTION("Transact(Assign, Open(AUTORTFM_DISABLE_SANITIZER_SCOPE(Assign)))")
+    {
+        AutoRTFM::Transact([&]
+        {
+            I = 42;
+
+            AutoRTFM::Open([&]
+            {
+                AUTORTFM_DISABLE_SANITIZER_SCOPE();
+
+                I = 24; // Sanitizer error silenced by AUTORTFM_DISABLE_SANITIZER_SCOPE()
+            });
+        });
+    }
+
     SECTION("Transact(Assign, Transact(Open(Assign)))")
     {
         AutoRTFM::Transact([&]
@@ -848,5 +868,9 @@ TEST_CASE("Open.Collision")
         REQUIRE(WarningContext.HasWarningSubstring(ExpectedFileAndLine));
 #endif // UE_BUILD_DEBUG
     }
+    else
+    {
+        REQUIRE(WarningContext.GetWarnings().empty());
+    }
 }
 #endif // AUTORTFM_SANITIZER
diff --git a/Engine/Source/Programs/AutoRTFMTests/Private/UECoreTests.cpp b/Engine/Source/Programs/AutoRTFMTests/Private/UECoreTests.cpp
--- a/Engine/Source/Programs/AutoRTFMTests/Private/UECoreTests.cpp
+++ b/Engine/Source/Programs/AutoRTFMTests/Private/UECoreTests.cpp
@@ -877,7 +877,7 @@ TEST_CASE("UECore.FUObjectItem")
 	SECTION("CreateStatID In Closed Then Again In Open")
 	{
 		// This test case performs a mixed open / closed write in Item.CreateStatID()
-		AUTORTFM_SCOPED_DISABLE_SANITIZER();
+		AUTORTFM_SCOPED_SANITIZER_MODE_DISABLE();
 
 		{
 			FUObjectItem Item;
diff --git a/Engine/Source/Runtime/AutoRTFM/Private/API.cpp b/Engine/Source/Runtime/AutoRTFM/Private/API.cpp
--- a/Engine/Source/Runtime/AutoRTFM/Private/API.cpp
+++ b/Engine/Source/Runtime/AutoRTFM/Private/API.cpp
@@ -30,7 +30,7 @@ namespace
 	// See `AutoRTFM::ForTheRuntime::CoinTossDisable` for implementation details.
 	float GAutoRTFMEnabledProbability = 100.0f;
 
-	AUTORTFM_SANITIZER_ONLY(AutoRTFM::ESanitizerMode GSanitizerMode = AutoRTFM::ESanitizerMode::Error);
+	AUTORTFM_SANITIZER_ONLY(AutoRTFM::Sanitizer::EMode GSanitizerMode = AutoRTFM::Sanitizer::EMode::Error);
 	AUTORTFM_SANITIZER_ONLY(bool GSanitizerRecordClosedWriteCallstacks = false);
 
 	bool GAutoRTFMEnsureOnInternalAbort = true;
@@ -85,6 +85,68 @@ namespace AutoRTFM
 
 namespace AutoRTFM
 {
+
+#if AUTORTFM_SANITIZER
+	namespace Sanitizer
+	{
+		EMode GetMode()
+		{
+			return GSanitizerMode;
+		}
+
+		void SetMode(EMode Mode)
+		{
+			if (GSanitizerMode != Mode)
+			{
+				GSanitizerMode = Mode;
+				if (GExternAPI.OnSanitizerModeChanged)
+				{
+					GExternAPI.OnSanitizerModeChanged();
+				}
+			}
+		}
+
+		void BeginDisable()
+		{
+			if (FContext* Context = FContext::Get())
+			{
+				if (FTransaction* Transaction = Context->GetCurrentTransaction())
+				{
+					Transaction->DisableSanitizer();
+				}
+			}
+		}
+
+		void EndDisable()
+		{
+			if (FContext* Context = FContext::Get())
+			{
+				if (FTransaction* Transaction = Context->GetCurrentTransaction())
+				{
+					Transaction->EnableSanitizer();
+				}
+			}
+		}
+
+		bool RecordClosedWriteCallstacks()
+		{
+			return GSanitizerRecordClosedWriteCallstacks;
+		}
+
+		void SetRecordClosedWriteCallstacks(bool bEnable)
+		{
+			if (GSanitizerRecordClosedWriteCallstacks != bEnable)
+			{
+				GSanitizerRecordClosedWriteCallstacks = bEnable;
+				if (GExternAPI.OnSanitizerRecordClosedWriteCallstacksChanged)
+				{
+					GExternAPI.OnSanitizerRecordClosedWriteCallstacksChanged();
+				}
+			}
+		}
+	}
+#endif // AUTORTFM_SANITIZER
+
 	namespace ForTheRuntime
 	{
 		bool SetAutoRTFMRuntime(EAutoRTFMEnabledState State)
@@ -320,43 +382,6 @@ namespace AutoRTFM
 			return false;
 #endif
 		}
-
-#if AUTORTFM_SANITIZER
-		ESanitizerMode GetSanitizerMode()
-		{
-			return GSanitizerMode;
-		}
-
-		void SetSanitizerMode(ESanitizerMode Mode)
-		{
-			if (GSanitizerMode != Mode)
-			{
-				GSanitizerMode = Mode;
-				if (GExternAPI.OnSanitizerModeChanged)
-				{
-					GExternAPI.OnSanitizerModeChanged();
-				}
-			}
-		}
-
-		bool SanitizerRecordClosedWriteCallstacks()
-		{
-			return GSanitizerRecordClosedWriteCallstacks;
-		}
-
-		void SetSanitizerRecordClosedWriteCallstacks(bool bEnable)
-		{
-			if (GSanitizerRecordClosedWriteCallstacks != bEnable)
-			{
-				GSanitizerRecordClosedWriteCallstacks = bEnable;
-				if (GExternAPI.OnSanitizerRecordClosedWriteCallstacksChanged)
-				{
-					GExternAPI.OnSanitizerRecordClosedWriteCallstacksChanged();
-				}
-			}
-		}
-#endif // AUTORTFM_SANITIZER
-
 	}
 }
 
@@ -438,17 +463,22 @@ AUTORTFM_INTERNAL autortfm_status RTFM_autortfm_close(void (*UninstrumentedWork)
 extern "C" UE_AUTORTFM_NOAUTORTFM UE_AUTORTFM_API void autortfm_pre_open(bool bNoSanitize)
 {
 #if AUTORTFM_SANITIZER
-	FContext::Get()->GetCurrentTransaction()->PushSanitizerMode(
-		bNoSanitize ? ESanitizerMode::Disabled : ForTheRuntime::GetSanitizerMode());
+	if (bNoSanitize)
+	{
+		FContext::Get()->GetCurrentTransaction()->DisableSanitizer();
+	}
 #endif // AUTORTFM_SANITIZER
 }
 
-extern "C" UE_AUTORTFM_NOAUTORTFM UE_AUTORTFM_API void autortfm_post_open()
+extern "C" UE_AUTORTFM_NOAUTORTFM UE_AUTORTFM_API void autortfm_post_open(bool bNoSanitize)
 {
 	FContext* Context = FContext::Get();
 
 #if AUTORTFM_SANITIZER
-	Context->GetCurrentTransaction()->PopSanitizerMode();
+	if (bNoSanitize)
+	{
+		FContext::Get()->GetCurrentTransaction()->EnableSanitizer();
+	}
 #endif // AUTORTFM_SANITIZER
 
 	if (Context->IsAborting())
@@ -490,7 +520,7 @@ extern "C" AUTORTFM_INTERNAL UE_AUTORTFM_API void autortfm_open_from_closed_no_s
 
 	FTransaction* Transaction = Context->GetCurrentTransaction();
 	
-	AUTORTFM_SANITIZER_ONLY(Transaction->PushSanitizerMode(ESanitizerMode::Disabled));
+	AUTORTFM_SANITIZER_ONLY(Transaction->DisableSanitizer());
 
 	Work(Arg);
 
@@ -501,7 +531,7 @@ extern "C" AUTORTFM_INTERNAL UE_AUTORTFM_API void autortfm_open_from_closed_no_s
 
 	AUTORTFM_ASSERT(Transaction == Context->GetCurrentTransaction());
 
-	AUTORTFM_SANITIZER_ONLY(Transaction->PopSanitizerMode());
+	AUTORTFM_SANITIZER_ONLY(Transaction->EnableSanitizer());
 }
 
 AUTORTFM_INTERNAL void RTFM_autortfm_record_open_write_err(void*, size_t)
diff --git a/Engine/Source/Runtime/AutoRTFM/Private/Transaction.cpp b/Engine/Source/Runtime/AutoRTFM/Private/Transaction.cpp
--- a/Engine/Source/Runtime/AutoRTFM/Private/Transaction.cpp
+++ b/Engine/Source/Runtime/AutoRTFM/Private/Transaction.cpp
@@ -97,10 +97,7 @@ void FTransaction::Reset()
 	DeferredPopAllOnCommitHandlers.Reset();
 	DeferredPopAllOnAbortHandlers.Reset();
 
-#if AUTORTFM_SANITIZER
-	ClosedWrites.Reset();
-	SanitizerModeStack.Clear();
-#endif
+	AUTORTFM_SANITIZER_ONLY(ClosedWrites.Reset());
 
 	AUTORTFM_ASSERT(IsFresh());
 }
diff --git a/Engine/Source/Runtime/AutoRTFM/Private/Transaction.h b/Engine/Source/Runtime/AutoRTFM/Private/Transaction.h
--- a/Engine/Source/Runtime/AutoRTFM/Private/Transaction.h
+++ b/Engine/Source/Runtime/AutoRTFM/Private/Transaction.h
@@ -102,19 +102,25 @@ public:
 #if AUTORTFM_SANITIZER
 	inline void SanitizerOpenWrite(void* Ptr, size_t Size);
 	
-	inline void PushSanitizerMode(ESanitizerMode NewMode)
-	{
-		SanitizerModeStack.Push(NewMode);
-	}
+	inline void DisableSanitizer() { SanitizerDisableCount++; }
 
-	inline void PopSanitizerMode()
+	inline void EnableSanitizer()
 	{
-		SanitizerModeStack.Pop();
+		AUTORTFM_ASSERT(SanitizerDisableCount > 0);;
+		SanitizerDisableCount--;
 	}
 
-	inline ESanitizerMode SanitizerMode() const
+	inline Sanitizer::EMode SanitizerMode() const
 	{
-		return SanitizerModeStack.IsEmpty() ? ForTheRuntime::GetSanitizerMode() : SanitizerModeStack.Back();
+		if (SanitizerDisableCount)
+		{
+			return Sanitizer::EMode::Disabled;
+		}
+		if (Parent)
+		{
+			return Parent->SanitizerMode();
+		}
+		return Sanitizer::GetMode();
 	}
 #endif
 
@@ -172,7 +178,7 @@ private:
 	FIntervalTree NewMemoryTracker;
 #if AUTORTFM_SANITIZER
 	using TClosedWriteTree = TIntervalTree<FCallstackTree::FHandle>;
-	TStack<ESanitizerMode, 8> SanitizerModeStack;
+	uint32_t SanitizerDisableCount = 0;
 	TClosedWriteTree ClosedWrites;
 #endif
 	FWriteLog WriteLog;
diff --git a/Engine/Source/Runtime/AutoRTFM/Private/TransactionInlines.h b/Engine/Source/Runtime/AutoRTFM/Private/TransactionInlines.h
--- a/Engine/Source/Runtime/AutoRTFM/Private/TransactionInlines.h
+++ b/Engine/Source/Runtime/AutoRTFM/Private/TransactionInlines.h
@@ -129,7 +129,7 @@ void FTransaction::RecordWrite(void* LogicalAddress, size_t Size, EWriteFlags Fl
 	if ((Flags & EWriteFlags::NoSanitize) != EWriteFlags::NoSanitize)
 	{
 		FCallstackTree::FHandle CallstackHandle = 
-			AutoRTFM::ForTheRuntime::SanitizerRecordClosedWriteCallstacks()
+			AutoRTFM::Sanitizer::RecordClosedWriteCallstacks()
 				? Context->GetCallstackHandle(/* FramesToSkip */ 0)
 				: FCallstackTree::InvalidHandle;
 		ClosedWrites.Insert(LogicalAddress, Size, CallstackHandle);
@@ -196,7 +196,7 @@ template<unsigned SIZE> AUTORTFM_NO_ASAN UE_AUTORTFM_FORCENOINLINE void FTransac
 
 #if AUTORTFM_SANITIZER
 	FCallstackTree::FHandle CallstackHandle =
-		AutoRTFM::ForTheRuntime::SanitizerRecordClosedWriteCallstacks()
+		AutoRTFM::Sanitizer::RecordClosedWriteCallstacks()
 			? Context->GetCallstackHandle(/* FramesToSkip */ 0)
 			: FCallstackTree::InvalidHandle;
 	ClosedWrites.Insert(LogicalAddress, SIZE, CallstackHandle);
@@ -360,9 +360,9 @@ void FTransaction::SanitizerOpenWrite(void* Ptr, size_t Size)
 		return;
 	}
 
-	ESanitizerMode SanitizerMode = this->SanitizerMode();
+	Sanitizer::EMode SanitizerMode = this->SanitizerMode();
 
-	if (SanitizerMode == ESanitizerMode::Disabled)
+	if (SanitizerMode == Sanitizer::EMode::Disabled)
 	{
 		return;
 	}
@@ -375,7 +375,7 @@ void FTransaction::SanitizerOpenWrite(void* Ptr, size_t Size)
 		// Fix this.
 		if (const FCallstackTree::FHandle* Callstack = Transaction->ClosedWrites.Get(Ptr, Size); AUTORTFM_UNLIKELY(Callstack))
 		{
-			const autortfm_log_severity Severity = (SanitizerMode == ESanitizerMode::Error) ? autortfm_log_error : autortfm_log_warn;
+			const autortfm_log_severity Severity = (SanitizerMode == Sanitizer::EMode::Error) ? autortfm_log_error : autortfm_log_warn;
 
 			AUTORTFM_LOG_AT_SEVERITY(Severity, "AutoRTFMSan: 0x%zu bytes at %p was modified in the closed and now the open.", Size, Ptr);
 			AUTORTFM_LOG_AT_SEVERITY(Severity, "This may lead to memory corruption if the transaction is aborted.");
@@ -391,7 +391,7 @@ void FTransaction::SanitizerOpenWrite(void* Ptr, size_t Size)
 			::AutoRTFM::LogWithCallstack(Severity, "AutoRTFMSan: The open write was at:");
 			AUTORTFM_LOG_AT_SEVERITY(Severity, "");
 
-			if (SanitizerMode == ESanitizerMode::Error)
+			if (SanitizerMode == Sanitizer::EMode::Error)
 			{
 				AUTORTFM_FATAL("<AutoRTFMSan fatal error>");
 			}
diff --git a/Engine/Source/Runtime/AutoRTFM/Public/AutoRTFM.h b/Engine/Source/Runtime/AutoRTFM/Public/AutoRTFM.h
--- a/Engine/Source/Runtime/AutoRTFM/Public/AutoRTFM.h
+++ b/Engine/Source/Runtime/AutoRTFM/Public/AutoRTFM.h
@@ -155,23 +155,70 @@ inline bool IsStatusAborting(EContextStatus Status)
 // An opaque unique identifier for a transaction.
 using TransactionID = autortfm_transaction_id;
 
-// An enumerator of AutoRTFM sanitizer levels.
-// The AutoRTFM sanitizer is used to detect modification by open-code to memory that was written by
-// a transaction. In this situation, aborting the transaction can corrupt memory as the undo will
-// overwrite the writes made in the open-code.
-enum class ESanitizerMode : uint8_t
+#if AUTORTFM_SANITIZER
+// AutoRTFM sanitizer API (requires the -fsanitize=autortfm compiler flag)
+namespace Sanitizer
 {
-	// Disable AutoRTFM sanitizer.
-	Disabled = autortfm_sanitizer_disabled,
+	// An enumerator of AutoRTFM sanitizer levels.
+	// The AutoRTFM sanitizer is used to detect modification by open-code to memory that was written by
+	// a transaction. In this situation, aborting the transaction can corrupt memory as the undo will
+	// overwrite the writes made in the open-code.
+	enum class EMode : uint8_t
+	{
+		// AutoRTFM sanitizer is disabled.
+		Disabled = autortfm_sanitizer_disabled,
 
-	// Enable AutoRTFM sanitizer.
-	// AutoRTFM sanitizer failures are treated as warnings.
-	Warn = autortfm_sanitizer_warn,
+		// AutoRTFM sanitizer is enabled and failures are treated as warnings.
+		Warn = autortfm_sanitizer_warn,
 
-	// Enable AutoRTFM sanitizer.
-	// AutoRTFM sanitizer failures are treated as fatal.
-	Error = autortfm_sanitizer_error,
-};
+		// AutoRTFM sanitizer is enabled and failures are treated as errors.
+		Error = autortfm_sanitizer_error,
+	};
+
+	// Returns the global AutoRTFM sanitizer mode currently in use.
+	UE_AUTORTFM_API EMode GetMode();
+
+	// Sets the global AutoRTFM sanitizer mode.
+	UE_AUTORTFM_API void SetMode(EMode Mode);
+
+	// Temporarily disables the AutoRTFM sanitizer for the current transaction
+	// and any child transactions.
+	// Must be paired with a call to EndDisable().
+	// Calls to BeginDisable() are internally counted and may be nested within
+	// other BeginDisable() / EndDisable() calls.
+	// This call is a no-op when called outside of a transaction.
+	UE_AUTORTFM_API AUTORTFM_OPEN void BeginDisable();
+
+	// Ends temporarily disabling the AutoRTFM sanitizer.
+	// See BeginDisable().
+	UE_AUTORTFM_API AUTORTFM_OPEN void EndDisable();
+
+	// Returns true if the AutoRTFM sanitizer records the callstack of closed writes.
+	UE_AUTORTFM_API bool RecordClosedWriteCallstacks();
+
+	// Sets whether the AutoRTFM sanitizer records the callstack of closed writes.
+	UE_AUTORTFM_API void SetRecordClosedWriteCallstacks(bool bEnable);
+
+	// A helper class that temporarily disables the AutoRTFM sanitizer for
+	// the lifetime of the object.
+	class FDisableScope
+	{
+	public:
+		UE_AUTORTFM_FORCEINLINE FDisableScope() { BeginDisable(); }
+		UE_AUTORTFM_FORCEINLINE ~FDisableScope() { EndDisable(); }
+	private:
+		FDisableScope(const FDisableScope&) = delete;
+		FDisableScope& operator = (const FDisableScope&) = delete;
+	};
+}
+#endif // AUTORTFM_SANITIZER
+
+#if AUTORTFM_SANITIZER
+	#define AUTORTFM_DISABLE_SANITIZER_SCOPE() \
+		::AutoRTFM::Sanitizer::FDisableScope UE_AUTORTFM_CONCAT(AutoRTFMSanitizerModeScope, __COUNTER__)
+#else
+	#define AUTORTFM_DISABLE_SANITIZER_SCOPE() static_assert(true) /* require semicolon */
+#endif
 
 #if UE_AUTORTFM
 namespace ForTheRuntime
@@ -745,8 +792,11 @@ template<typename Type> UE_AUTORTFM_CRITICAL_INLINE void RecordOpenWrite(Type* P
 template<typename RefType, typename ValueType>
 AUTORTFM_DISABLE_NO_SANITIZE void Assign(RefType& Ref, ValueType&& NewValue, EWriteFlags Flags = EWriteFlags::Default)
 {
-	autortfm_record_open_write_with_flags(&Ref, sizeof(Ref), static_cast<autortfm_write_flags>(Flags | EWriteFlags::NoSanitize));
-	Ref = std::forward<ValueType>(NewValue);
+	autortfm_record_open_write_with_flags(&Ref, sizeof(Ref), static_cast<autortfm_write_flags>(Flags));
+	{
+		AUTORTFM_DISABLE_SANITIZER_SCOPE();
+		Ref = std::forward<ValueType>(NewValue);
+	}
 }
 #else
 template<typename RefType, typename ValueType>
@@ -932,20 +982,6 @@ namespace ForTheRuntime
 	// Returns true if we should also retry nested transactions.
 	UE_AUTORTFM_API bool ShouldRetryNestedTransactionsToo();
 
-#if AUTORTFM_SANITIZER
-	// Returns the AutoRTFM Sanitizer mode currently in use.
-	UE_AUTORTFM_API ESanitizerMode GetSanitizerMode();
-
-	// Sets the AutoRTFM Sanitizer mode.
-	UE_AUTORTFM_API void SetSanitizerMode(ESanitizerMode Mode);
-
-	// Returns true if the AutoRTFM Sanitizer records the callstack of closed writes.
-	UE_AUTORTFM_API bool SanitizerRecordClosedWriteCallstacks();
-
-	// Sets whether the AutoRTFM Sanitizer records the callstack of closed writes.
-	UE_AUTORTFM_API void SetSanitizerRecordClosedWriteCallstacks(bool bEnable);
-#endif // AUTORTFM_SANITIZER
-
 	// Manually create a new transaction from open code and push it as a transaction nest.
 	// Can only be called within an already active parent transaction (EG. this cannot start
 	// a transaction nest itself).
diff --git a/Engine/Source/Runtime/Core/Private/AutoRTFM/AutoRTFMUE.cpp b/Engine/Source/Runtime/Core/Private/AutoRTFM/AutoRTFMUE.cpp
--- a/Engine/Source/Runtime/Core/Private/AutoRTFM/AutoRTFMUE.cpp
+++ b/Engine/Source/Runtime/Core/Private/AutoRTFM/AutoRTFMUE.cpp
@@ -65,15 +65,15 @@ void OnAutoRTFMRetryTransactionsChanged()
 #if AUTORTFM_SANITIZER
 void OnAutoRTFMSanitizerModeChanged()
 {
-	switch (AutoRTFM::ForTheRuntime::GetSanitizerMode())
+	switch (AutoRTFM::Sanitizer::GetMode())
 	{
-	case AutoRTFM::ESanitizerMode::Disabled:
+	case AutoRTFM::Sanitizer::EMode::Disabled:
 		UE_LOG(LogAutoRTFM, Log, TEXT("AutoRTFM Sanitizer: disabled"));
 		return FGenericCrashContext::SetGameData(TEXT("AutoRTFMSanitizerMode"), TEXT("Disabled"));
-	case AutoRTFM::ESanitizerMode::Warn:
+	case AutoRTFM::Sanitizer::EMode::Warn:
 		UE_LOG(LogAutoRTFM, Log, TEXT("AutoRTFM Sanitizer: enabled as warning"));
 		return FGenericCrashContext::SetGameData(TEXT("AutoRTFMSanitizerMode"), TEXT("Warn"));
-	case AutoRTFM::ESanitizerMode::Error:
+	case AutoRTFM::Sanitizer::EMode::Error:
 		UE_LOG(LogAutoRTFM, Log, TEXT("AutoRTFM Sanitizer: enabled as error"));
 		return FGenericCrashContext::SetGameData(TEXT("AutoRTFMSanitizerMode"), TEXT("Enabled"));
 	}
@@ -82,7 +82,7 @@ void OnAutoRTFMSanitizerModeChanged()
 void OnAutoRTFMSanitizerRecordClosedWriteCallstacksChanged()
 {
 	UE_LOG(LogAutoRTFM, Log, TEXT("AutoRTFM Sanitizer: %s recording of closed write callstacks"),
-		AutoRTFM::ForTheRuntime::SanitizerRecordClosedWriteCallstacks() ? TEXT("enabled") : TEXT("disabled"));
+		AutoRTFM::Sanitizer::RecordClosedWriteCallstacks() ? TEXT("enabled") : TEXT("disabled"));
 }
 
 #endif // AUTORTFM_SANITIZER
@@ -174,24 +174,24 @@ static FAutoConsoleVariable CVarAutoRTFMRetryTransactions(
 #if AUTORTFM_SANITIZER
 static FAutoConsoleVariable CVarAutoRTFMSanitizerMode(
 	TEXT("AutoRTFMSanitizerMode"),
-	static_cast<int>(AutoRTFM::ForTheRuntime::GetSanitizerMode()),
+	static_cast<int>(AutoRTFM::Sanitizer::GetMode()),
 	TEXT("Detects potential memory corruption due to writes made both by a transaction and open-code"),
 	FConsoleVariableDelegate::CreateLambda([](IConsoleVariable* Variable)
 	{
 		FString Value = Variable->GetString().ToLower();
 		if (IsDisable(Value))
 		{
-			AutoRTFM::ForTheRuntime::SetSanitizerMode(AutoRTFM::ESanitizerMode::Disabled);
+			AutoRTFM::Sanitizer::SetMode(AutoRTFM::Sanitizer::EMode::Disabled);
 			return;
 		}
 		if (Value == "warn")
 		{
-			AutoRTFM::ForTheRuntime::SetSanitizerMode(AutoRTFM::ESanitizerMode::Warn);
+			AutoRTFM::Sanitizer::SetMode(AutoRTFM::Sanitizer::EMode::Warn);
 			return;
 		}
 		if (Value == "error" || IsEnable(Value))
 		{
-			AutoRTFM::ForTheRuntime::SetSanitizerMode(AutoRTFM::ESanitizerMode::Error);
+			AutoRTFM::Sanitizer::SetMode(AutoRTFM::Sanitizer::EMode::Error);
 			return;
 		}
 		UE_LOG(LogAutoRTFM, Fatal, TEXT("'AutoRTFMSanitizerMode' CVar was set to '%s' which is not one of 'disable', 'warn' or 'error'!"), *Value);
@@ -201,19 +201,19 @@ static FAutoConsoleVariable CVarAutoRTFMSanitizerMode(
 
 static FAutoConsoleVariable CVarAutoRTFMSanitizerRecordWriteCallstacks(
 	TEXT("AutoRTFMSanitizerRecordWriteCallstacks"),
-	AutoRTFM::ForTheRuntime::SanitizerRecordClosedWriteCallstacks(),
+	AutoRTFM::Sanitizer::RecordClosedWriteCallstacks(),
 	TEXT("Enables the recording of callstacks for closed writes, providing more information on AutoRTFMSan failures"),
 	FConsoleVariableDelegate::CreateLambda([](IConsoleVariable* Variable)
 	{
 		FString Value = Variable->GetString().ToLower();
 		if (IsDisable(Value))
 		{
-			AutoRTFM::ForTheRuntime::SetSanitizerRecordClosedWriteCallstacks(false);
+			AutoRTFM::Sanitizer::SetRecordClosedWriteCallstacks(false);
 			return;
 		}
 		if (IsEnable(Value))
 		{
-			AutoRTFM::ForTheRuntime::SetSanitizerRecordClosedWriteCallstacks(true);
+			AutoRTFM::Sanitizer::SetRecordClosedWriteCallstacks(true);
 			return;
 		}
 		UE_LOG(LogAutoRTFM, Fatal, TEXT("'AutoRTFMSanitizerRecordWriteCallstacks' CVar was set to '%s' which is not one of 'enable' or 'disable'!"), *Value);
