[AutoRTFM] Improve sanitizer failure output

Ensure that the callstacks are trimmed to exclude the callstack gathering logic itself.
Use the same callstack utilities for closed and open backtraces, so that output formatting is identical for both.
Add a cache of already-reported stack traces to that we don't spam the same failure callstacks.
Move the bulky message printing logic out of the inline function and into a Transaction.cpp method.
Include a message saying when closed write callstack recordining is not enabled.
Include some repeated '*' characters to help visually separate multiple failures printed when the mode is warn.

#Synced-CL 49102748

diff --git a/Engine/Source/Runtime/AutoRTFM/Private/Context.cpp b/Engine/Source/Runtime/AutoRTFM/Private/Context.cpp
--- a/Engine/Source/Runtime/AutoRTFM/Private/Context.cpp
+++ b/Engine/Source/Runtime/AutoRTFM/Private/Context.cpp
@@ -644,14 +644,25 @@ void FContext::DumpState() const
 
 
 #if AUTORTFM_SANITIZER
-FCallstackTree::FHandle FContext::GetCallstackHandle(size_t FramesToSkip)
+FCallstackTree::FHandle FContext::GetCallstackHandle(size_t FramesToSkip /* = 0 */)
 {
-	static constexpr size_t MaxCallstack = 64;
+	static constexpr size_t MaxCallstack = 128;
+	autortfm_callstack_frame Callstack[MaxCallstack];
+	size_t const NumCaptured = GExternAPI.CaptureCallstack(MaxCallstack, Callstack);
 
-	FramesToSkip += 3; // Top frame is caller.
+	// Attempt to trim off callstack frames under this function.
+	if (autortfm_callstack_frame ReturnAddress = reinterpret_cast<autortfm_callstack_frame>(__builtin_return_address(0)))
+	{
+		for (size_t I = 0, N = std::min(NumCaptured, static_cast<size_t>(8)); I < N; I++)
+		{
+			if (ReturnAddress == Callstack[I])
+			{
+				FramesToSkip += I;
+				break;
+			}
+		}
+	}
 
-	autortfm_callstack_frame Callstack[MaxCallstack];
-	size_t NumCaptured = GExternAPI.CaptureCallstack(MaxCallstack, Callstack);
 	if (NumCaptured <= FramesToSkip)
 	{
 		return FCallstackTree::InvalidHandle;
diff --git a/Engine/Source/Runtime/AutoRTFM/Private/Context.h b/Engine/Source/Runtime/AutoRTFM/Private/Context.h
--- a/Engine/Source/Runtime/AutoRTFM/Private/Context.h
+++ b/Engine/Source/Runtime/AutoRTFM/Private/Context.h
@@ -138,7 +138,7 @@ public:
 	static FStackRange GetThreadStackRange();
 
 #if AUTORTFM_SANITIZER
-	FCallstackTree::FHandle GetCallstackHandle(size_t FramesToSkip);
+	FCallstackTree::FHandle GetCallstackHandle(size_t FramesToSkip = 0);
 	void LogCallstack(FCallstackTree::FHandle, autortfm_log_severity);
 #endif
 
diff --git a/Engine/Source/Runtime/AutoRTFM/Private/Transaction.cpp b/Engine/Source/Runtime/AutoRTFM/Private/Transaction.cpp
--- a/Engine/Source/Runtime/AutoRTFM/Private/Transaction.cpp
+++ b/Engine/Source/Runtime/AutoRTFM/Private/Transaction.cpp
@@ -343,6 +343,68 @@ bool FTransaction::AttemptToCommitOuterNest()
 	return true;
 }
 
+#if AUTORTFM_SANITIZER
+void FTransaction::SanitizerReportConflictingOpenWrite(void* Ptr, size_t Size, FCallstackTree::FHandle ClosedWriteCallstack)
+{
+	Sanitizer::EMode SanitizerMode = this->SanitizerMode();
+	
+	const autortfm_log_severity Severity = (SanitizerMode == Sanitizer::EMode::Error) ? autortfm_log_error : autortfm_log_warn;
+
+	FCallstackTree::FHandle const OpenWriteCallstack = Context->GetCallstackHandle(/* FramesToSkip */ 2);
+	
+	if (Severity == autortfm_log_warn)
+	{
+		struct FCallstackPair
+		{
+			FCallstackTree::FHandle Open;
+			FCallstackTree::FHandle Closed;
+
+			struct FHasher
+			{
+				size_t operator() (FCallstackPair Pair) const
+				{
+					return (static_cast<size_t>(Pair.Open) << 32) | static_cast<size_t>(Pair.Closed);
+				}
+			};
+			bool operator == (const FCallstackPair&) const = default;
+		};
+		static std::unordered_set<FCallstackPair, FCallstackPair::FHasher> Reported;
+		if (!Reported.emplace(OpenWriteCallstack, ClosedWriteCallstack).second)
+		{
+			return; // Already reported
+		}
+	}
+
+	AUTORTFM_LOG_AT_SEVERITY(Severity, "");
+	AUTORTFM_LOG_AT_SEVERITY(Severity, "********************************************************************************************");
+	AUTORTFM_LOG_AT_SEVERITY(Severity, "AutoRTFMSan: 0x%zu bytes at %p was modified in the closed and now the open.", Size, Ptr);
+	AUTORTFM_LOG_AT_SEVERITY(Severity, "This may lead to memory corruption if the transaction is aborted.");
+	AUTORTFM_LOG_AT_SEVERITY(Severity, "");
+
+	if (ClosedWriteCallstack != FCallstackTree::InvalidHandle)
+	{
+		AUTORTFM_LOG_AT_SEVERITY(Severity, "AutoRTFMSan: The last closed write to this address was at:");
+		Context->LogCallstack(ClosedWriteCallstack, Severity);
+		AUTORTFM_LOG_AT_SEVERITY(Severity, "");
+	}
+	else if (!AutoRTFM::Sanitizer::RecordClosedWriteCallstacks())
+	{
+		AUTORTFM_LOG_AT_SEVERITY(Severity, "AutoRTFMSan: Closed write callstack recording is not enabled.");
+		AUTORTFM_LOG_AT_SEVERITY(Severity, "AutoRTFMSan: Enable to see callstack of the last closed write.");
+	}
+
+	AUTORTFM_LOG_AT_SEVERITY(Severity, "AutoRTFMSan: The open write was at:");
+	Context->LogCallstack(OpenWriteCallstack, Severity);
+	AUTORTFM_LOG_AT_SEVERITY(Severity, "********************************************************************************************");
+	AUTORTFM_LOG_AT_SEVERITY(Severity, "");
+
+	if (SanitizerMode == Sanitizer::EMode::Error)
+	{
+		AUTORTFM_FATAL("<AutoRTFMSan fatal error>");
+	}
+}
+#endif // AUTORTFM_SANITIZER
+
 } // namespace AutoRTFM
 
 #endif // defined(__AUTORTFM) && __AUTORTFM
diff --git a/Engine/Source/Runtime/AutoRTFM/Private/Transaction.h b/Engine/Source/Runtime/AutoRTFM/Private/Transaction.h
--- a/Engine/Source/Runtime/AutoRTFM/Private/Transaction.h
+++ b/Engine/Source/Runtime/AutoRTFM/Private/Transaction.h
@@ -101,7 +101,7 @@ public:
 
 #if AUTORTFM_SANITIZER
 	inline void SanitizerOpenWrite(void* Ptr, size_t Size);
-	
+
 	inline void DisableSanitizer() { SanitizerDisableCount++; }
 
 	inline void EnableSanitizer()
@@ -136,6 +136,10 @@ private:
 
 	bool ShouldRecordWrite(void* LogicalAddress) const;
 
+#if AUTORTFM_SANITIZER
+	void SanitizerReportConflictingOpenWrite(void* Ptr, size_t Size, FCallstackTree::FHandle ClosedWriteCallstack);
+#endif
+
 	FContext* Context;
 
 	// This field is used for two purposes:
diff --git a/Engine/Source/Runtime/AutoRTFM/Private/TransactionInlines.h b/Engine/Source/Runtime/AutoRTFM/Private/TransactionInlines.h
--- a/Engine/Source/Runtime/AutoRTFM/Private/TransactionInlines.h
+++ b/Engine/Source/Runtime/AutoRTFM/Private/TransactionInlines.h
@@ -130,7 +130,7 @@ void FTransaction::RecordWrite(void* LogicalAddress, size_t Size, EWriteFlags Fl
 	{
 		FCallstackTree::FHandle CallstackHandle = 
 			AutoRTFM::Sanitizer::RecordClosedWriteCallstacks()
-				? Context->GetCallstackHandle(/* FramesToSkip */ 0)
+				? Context->GetCallstackHandle()
 				: FCallstackTree::InvalidHandle;
 		ClosedWrites.Insert(LogicalAddress, Size, CallstackHandle);
 	}
@@ -197,7 +197,7 @@ template<unsigned SIZE> AUTORTFM_NO_ASAN UE_AUTORTFM_FORCENOINLINE void FTransac
 #if AUTORTFM_SANITIZER
 	FCallstackTree::FHandle CallstackHandle =
 		AutoRTFM::Sanitizer::RecordClosedWriteCallstacks()
-			? Context->GetCallstackHandle(/* FramesToSkip */ 0)
+			? Context->GetCallstackHandle()
 			: FCallstackTree::InvalidHandle;
 	ClosedWrites.Insert(LogicalAddress, SIZE, CallstackHandle);
 #endif
@@ -360,9 +360,7 @@ void FTransaction::SanitizerOpenWrite(void* Ptr, size_t Size)
 		return;
 	}
 
-	Sanitizer::EMode SanitizerMode = this->SanitizerMode();
-
-	if (SanitizerMode == Sanitizer::EMode::Disabled)
+	if (this->SanitizerMode() == Sanitizer::EMode::Disabled)
 	{
 		return;
 	}
@@ -375,30 +373,11 @@ void FTransaction::SanitizerOpenWrite(void* Ptr, size_t Size)
 		// Fix this.
 		if (const FCallstackTree::FHandle* Callstack = Transaction->ClosedWrites.Get(Ptr, Size); AUTORTFM_UNLIKELY(Callstack))
 		{
-			const autortfm_log_severity Severity = (SanitizerMode == Sanitizer::EMode::Error) ? autortfm_log_error : autortfm_log_warn;
-
-			AUTORTFM_LOG_AT_SEVERITY(Severity, "AutoRTFMSan: 0x%zu bytes at %p was modified in the closed and now the open.", Size, Ptr);
-			AUTORTFM_LOG_AT_SEVERITY(Severity, "This may lead to memory corruption if the transaction is aborted.");
-			AUTORTFM_LOG_AT_SEVERITY(Severity, "");
-
-			if (*Callstack != FCallstackTree::InvalidHandle)
-			{
-				AUTORTFM_LOG_AT_SEVERITY(Severity, "AutoRTFMSan: The last closed write to this address was at:");
-				Context->LogCallstack(*Callstack, Severity);
-				AUTORTFM_LOG_AT_SEVERITY(Severity, "");
-			}
-
-			::AutoRTFM::LogWithCallstack(Severity, "AutoRTFMSan: The open write was at:");
-			AUTORTFM_LOG_AT_SEVERITY(Severity, "");
-
-			if (SanitizerMode == Sanitizer::EMode::Error)
-			{
-				AUTORTFM_FATAL("<AutoRTFMSan fatal error>");
-			}
+			SanitizerReportConflictingOpenWrite(Ptr, Size, *Callstack);
 		}
 	}
 }
-#endif
+#endif // AUTORTFM_SANITIZER
 
 } // namespace AutoRTFM
 
