[AutoRTFM] Add AutoRTFM shutdown method

Call this when UE is shutting down.
Have the shutdown print the number of AutoRTFM sanitizer failures reported at exit.

#Synced-CL 49163605

diff --git a/Engine/Source/Runtime/AutoRTFM/Private/API.cpp b/Engine/Source/Runtime/AutoRTFM/Private/API.cpp
--- a/Engine/Source/Runtime/AutoRTFM/Private/API.cpp
+++ b/Engine/Source/Runtime/AutoRTFM/Private/API.cpp
@@ -688,14 +688,27 @@ extern "C" void autortfm_initialize(const autortfm_extern_api* ExternAPI) noexce
 
 	AUTORTFM_ASSERT(ExternAPI);
 	AUTORTFM_ASSERT(ExternAPI->Allocate);
-	AUTORTFM_ASSERT(ExternAPI->AllocateZeroed);
 	AUTORTFM_ASSERT(ExternAPI->Reallocate);
+	AUTORTFM_ASSERT(ExternAPI->AllocateZeroed);
 	AUTORTFM_ASSERT(ExternAPI->Free);
+	AUTORTFM_ASSERT(ExternAPI->Log);
+	AUTORTFM_ASSERT(ExternAPI->LogWithCallstack);
+	AUTORTFM_ASSERT(ExternAPI->EnsureFailure);
+	AUTORTFM_ASSERT(ExternAPI->IsLogActive);
+	AUTORTFM_ASSERT(ExternAPI->CaptureCallstack);
+	AUTORTFM_ASSERT(ExternAPI->LogCallstack);
+
 	GExternAPI = *ExternAPI;
 	FContext::Create();
 	ProcessAllPendingOpenToClosedRegistrations();
 }
 
+extern "C" void autortfm_shutdown() noexcept
+{
+	AutoRTFM::UnreachableIfClosed("TODO: Mark this as AUTORTFM_DISABLED and fixup callers");
+	FContext::Destroy();
+}
+
 // Each function will be forked by the compiler into an open and closed variant.
 // autortfm_is_closed() is used to branch to the closed variants declared above.
 extern "C" bool autortfm_is_transactional() noexcept
@@ -959,11 +972,6 @@ extern "C" UE_AUTORTFM_ALWAYS_OPEN bool autortfm_is_on_current_transaction_stack
 	return false;
 }
 
-void ForTheRuntime::Initialize(const ForTheRuntime::FExternAPI& ExternAPI)
-{
-	autortfm_initialize(&ExternAPI);
-}
-
 void ForTheRuntime::CascadingAbortTransactionInternal()
 {
 	AUTORTFM_ASSERT(autortfm_is_closed());
@@ -1145,6 +1153,16 @@ void ForTheRuntime::RedirectedStore(uint32_t AddressSpace, uint64_t DestAddress,
 	}
 }
 
+void Initialize(const FExternAPI& ExternAPI)
+{
+	autortfm_initialize(&ExternAPI);
+}
+
+void Shutdown()
+{
+	autortfm_shutdown();
+}
+
 void RegisterHeapRedirectCallbacks(FHeapRedirectCallbacks Callbacks)
 {
 	AUTORTFM_ASSERT(GNumRegisteredHeapRedirects < GMaxRegisteredHeapRedirects);
diff --git a/Engine/Source/Runtime/AutoRTFM/Private/Context.cpp b/Engine/Source/Runtime/AutoRTFM/Private/Context.cpp
--- a/Engine/Source/Runtime/AutoRTFM/Private/Context.cpp
+++ b/Engine/Source/Runtime/AutoRTFM/Private/Context.cpp
@@ -39,6 +39,17 @@ FContext* FContext::Create()
 	return Instance;
 }
 
+void FContext::Destroy()
+{
+	AUTORTFM_ENSURE(Instance != nullptr);
+	FContext* Context = Instance;
+	// Null instance before deleting the context, as freeing the context may
+	// attempt to use the context.
+	Instance = nullptr; 
+	Context->~FContext();
+	AutoRTFM::Free(Context);
+}
+
 void ResetAutoRTFMMetrics()
 {
 	GAutoRTFMMetrics = FAutoRTFMMetrics{};
@@ -50,6 +61,16 @@ FAutoRTFMMetrics GetAutoRTFMMetrics()
 	return GAutoRTFMMetrics;
 }
 
+FContext::~FContext()
+{
+#if AUTORTFM_SANITIZER
+	if (NumSanitizerIssues)
+	{
+		AUTORTFM_WARN("AutoRTFMSan: %zu sanitizer issues detected", NumSanitizerIssues);
+	}
+#endif
+}
+
 bool FContext::IsTransactional() const
 {
     return GetStatus() == EContextStatus::OnTrack;
diff --git a/Engine/Source/Runtime/AutoRTFM/Private/Context.h b/Engine/Source/Runtime/AutoRTFM/Private/Context.h
--- a/Engine/Source/Runtime/AutoRTFM/Private/Context.h
+++ b/Engine/Source/Runtime/AutoRTFM/Private/Context.h
@@ -30,6 +30,7 @@ class AUTORTFM_DISABLE_NO_SANITIZE FContext final
 public:
 	static FContext* Create();
 	static FContext* Get();
+	static void Destroy();
 
     // This is public API
     ETransactionResult Transact(void (*UninstrumentedFunction)(void*), void (*InstrumentedFunction)(void*), void* Arg);
@@ -140,6 +141,7 @@ public:
 #if AUTORTFM_SANITIZER
 	FCallstackTree::FHandle GetCallstackHandle(size_t FramesToSkip = 0);
 	void LogCallstack(FCallstackTree::FHandle, autortfm_log_severity);
+	void IncrementNumSanitizerIssues() { NumSanitizerIssues++; }
 #endif
 
 private:
@@ -147,6 +149,7 @@ private:
 
 	FContext() { Reset(); }
     FContext(const FContext&) = delete;
+	~FContext();
 
 	void MaterializeDeferredTransactions();
 	void StartNonDeferredTransaction();
@@ -179,6 +182,7 @@ private:
 
 #if AUTORTFM_SANITIZER
 	FCallstackTree CallstackTree;
+	size_t NumSanitizerIssues = 0;
 #endif
 };
 
diff --git a/Engine/Source/Runtime/AutoRTFM/Private/Transaction.cpp b/Engine/Source/Runtime/AutoRTFM/Private/Transaction.cpp
--- a/Engine/Source/Runtime/AutoRTFM/Private/Transaction.cpp
+++ b/Engine/Source/Runtime/AutoRTFM/Private/Transaction.cpp
@@ -375,6 +375,8 @@ void FTransaction::SanitizerReportConflictingOpenWrite(void* Ptr, size_t Size, F
 		}
 	}
 
+	Context->IncrementNumSanitizerIssues();
+
 	AUTORTFM_LOG_AT_SEVERITY(Severity, "");
 	AUTORTFM_LOG_AT_SEVERITY(Severity, "********************************************************************************************");
 	AUTORTFM_LOG_AT_SEVERITY(Severity, "AutoRTFMSan: 0x%zu bytes at %p was modified in the closed and now the open.", Size, Ptr);
diff --git a/Engine/Source/Runtime/AutoRTFM/Public/AutoRTFM.h b/Engine/Source/Runtime/AutoRTFM/Public/AutoRTFM.h
--- a/Engine/Source/Runtime/AutoRTFM/Public/AutoRTFM.h
+++ b/Engine/Source/Runtime/AutoRTFM/Public/AutoRTFM.h
@@ -223,8 +223,6 @@ namespace Sanitizer
 #if UE_AUTORTFM
 namespace ForTheRuntime
 {
-	using FExternAPI = autortfm_extern_api;
-	UE_AUTORTFM_API void Initialize(const FExternAPI& ExternAPI);
 	UE_AUTORTFM_API void CascadingAbortTransactionInternal();
 	UE_AUTORTFM_API void CascadingRetryTransactionInternal();
 	UE_AUTORTFM_API void OnCommitInternal(TTask<void()>&& Work);
@@ -266,6 +264,29 @@ auto AutoRTFMLookupInstrumentedFunctorInvoker(const FunctorType& Functor) -> voi
 }
 #endif
 
+using FExternAPI = autortfm_extern_api;
+
+// Initialize the AutoRTFM library.
+// Must only be called once for the lifetime of the application.
+#if UE_AUTORTFM
+UE_AUTORTFM_API void Initialize(const FExternAPI& ExternAPI);
+#else
+UE_AUTORTFM_CRITICAL_INLINE void Initialize(const FExternAPI& ExternAPI)
+{
+	UE_AUTORTFM_UNUSED(ExternAPI);
+}
+#endif
+
+// Shutdown the AutoRTFM library.
+// Must only be called once for the lifetime of the application, after autortfm_initialize().
+#if UE_AUTORTFM
+UE_AUTORTFM_API void Shutdown();
+#else
+UE_AUTORTFM_CRITICAL_INLINE void Shutdown()
+{
+}
+#endif
+
 // Tells if we are currently running in a transaction. This will return true in an open nest
 // (see `Open`). This function is handled specially in the compiler and will be constant folded
 // as true in closed code, or preserved as a function call in open code.
diff --git a/Engine/Source/Runtime/AutoRTFM/Public/AutoRTFM/CAPI.h b/Engine/Source/Runtime/AutoRTFM/Public/AutoRTFM/CAPI.h
--- a/Engine/Source/Runtime/AutoRTFM/Public/AutoRTFM/CAPI.h
+++ b/Engine/Source/Runtime/AutoRTFM/Public/AutoRTFM/CAPI.h
@@ -186,6 +186,7 @@ typedef struct
 
 #if UE_AUTORTFM_ENABLED
 // Initialize the AutoRTFM library.
+// Must only be called once for the lifetime of the application.
 // Parameters:
 //   ExternAPI - Function pointers used by AutoRTFM for heap allocations, etc.
 //               Must be non-null.
@@ -197,6 +198,16 @@ UE_AUTORTFM_CRITICAL_INLINE void autortfm_initialize(const autortfm_extern_api*
 }
 #endif
 
+#if UE_AUTORTFM_ENABLED
+// Shutdown the AutoRTFM library.
+// Must only be called once for the lifetime of the application, after autortfm_initialize().
+UE_AUTORTFM_API void autortfm_shutdown() AUTORTFM_NOEXCEPT;
+#else
+UE_AUTORTFM_CRITICAL_INLINE void autortfm_shutdown() AUTORTFM_NOEXCEPT
+{
+}
+#endif
+
 #if UE_AUTORTFM_ENABLED
 // Note: There is no implementation of this function.
 // The AutoRTFM compiler will replace all calls to this function with a constant boolean value.
diff --git a/Engine/Source/Runtime/Core/Private/AutoRTFM/AutoRTFMUE.cpp b/Engine/Source/Runtime/Core/Private/AutoRTFM/AutoRTFMUE.cpp
--- a/Engine/Source/Runtime/Core/Private/AutoRTFM/AutoRTFMUE.cpp
+++ b/Engine/Source/Runtime/Core/Private/AutoRTFM/AutoRTFMUE.cpp
@@ -417,7 +417,7 @@ static void AutoRTFMFree(void* Pointer)
 
 CORE_API void AutoRTFM::InitializeForUE(RollbackWriteFn RollbackWrite /* = nullptr */)
 {
-	AutoRTFM::ForTheRuntime::FExternAPI ExternAPI
+	AutoRTFM::FExternAPI ExternAPI
 	{
 		.Allocate = AutoRTFMAllocate,
 		.Reallocate = AutoRTFMReallocate,
@@ -438,7 +438,7 @@ CORE_API void AutoRTFM::InitializeForUE(RollbackWriteFn RollbackWrite /* = nullp
 		.RollbackWrite = RollbackWrite,
 	};
 	
-	AutoRTFM::ForTheRuntime::Initialize(ExternAPI);
+	AutoRTFM::Initialize(ExternAPI);
 
 	// Call the OnAutoRTFMXXXChanged() handlers now so that values are logged
 	// and crash context data is set.
diff --git a/Engine/Source/Runtime/Launch/Private/LaunchEngineLoop.cpp b/Engine/Source/Runtime/Launch/Private/LaunchEngineLoop.cpp
--- a/Engine/Source/Runtime/Launch/Private/LaunchEngineLoop.cpp
+++ b/Engine/Source/Runtime/Launch/Private/LaunchEngineLoop.cpp
@@ -7001,6 +7001,8 @@ void FEngineLoop::AppExit()
 	FInternationalization::TearDown();
 
 	FTraceAuxiliary::Shutdown();
+	
+	AutoRTFM::Shutdown();
 }
 
 void FEngineLoop::PostInitRHI()
