[AutoRTFM] Enable clang-format.

I attempted to minimize the impact by (substantially!) adjusting the clang-format rules compared to VerseVM, but it's still a large diff. Unfortunately we were inconsistent with indentation for namespaces and preprocessor commands, so large re-indentation diffs were unavoidable.

This CL also fixes a bug which caused the clang-format python script to fail in the presence of UTF8, and removes an obsolete reference to a "CoreVerseVM" directory which was moved long ago.

#Synced-CL 49209554

diff --git a/Engine/Source/Runtime/AutoRTFM/Private/API.cpp b/Engine/Source/Runtime/AutoRTFM/Private/API.cpp
--- a/Engine/Source/Runtime/AutoRTFM/Private/API.cpp
+++ b/Engine/Source/Runtime/AutoRTFM/Private/API.cpp
@@ -15,71 +15,70 @@ static_assert(UE_AUTORTFM_ENABLED, "AutoRTFM/API.cpp requires the compiler flag
 
 namespace
 {
-	// Move this to a local only and use functions to access this
+// Move this to a local only and use functions to access this
 #if UE_AUTORTFM_ENABLED_RUNTIME_BY_DEFAULT
-	int GAutoRTFMRuntimeEnabled = AutoRTFM::ForTheRuntime::EAutoRTFMEnabledState::AutoRTFM_EnabledByDefault;
+int GAutoRTFMRuntimeEnabled = AutoRTFM::ForTheRuntime::EAutoRTFMEnabledState::AutoRTFM_EnabledByDefault;
 #else
-	int GAutoRTFMRuntimeEnabled = AutoRTFM::ForTheRuntime::EAutoRTFMEnabledState::AutoRTFM_DisabledByDefault;
-#endif // UE_AUTORTFM_ENABLED_RUNTIME_BY_DEFAULT
+int GAutoRTFMRuntimeEnabled = AutoRTFM::ForTheRuntime::EAutoRTFMEnabledState::AutoRTFM_DisabledByDefault;
+#endif  // UE_AUTORTFM_ENABLED_RUNTIME_BY_DEFAULT
 
-	int GAutoRTFMInternalAbortAction = AutoRTFM::ForTheRuntime::EAutoRTFMInternalAbortActionState::Crash;
-	int GAutoRTFMRetryTransactions = AutoRTFM::ForTheRuntime::EAutoRTFMRetryTransactionState::NoRetry;
+int GAutoRTFMInternalAbortAction = AutoRTFM::ForTheRuntime::EAutoRTFMInternalAbortActionState::Crash;
+int GAutoRTFMRetryTransactions = AutoRTFM::ForTheRuntime::EAutoRTFMRetryTransactionState::NoRetry;
 
-	// Set the percentage chance [0..100] that AutoRTFM will be enabled.
-	// 100.0 means that AutoRTFM will always be enabled; 1.0 means that AutoRTFM has a 1% chance of being enabled.
-	// See `AutoRTFM::ForTheRuntime::CoinTossDisable` for implementation details.
-	float GAutoRTFMEnabledProbability = 100.0f;
+// Set the percentage chance [0..100] that AutoRTFM will be enabled.
+// 100.0 means that AutoRTFM will always be enabled; 1.0 means that AutoRTFM has a 1% chance of being enabled.
+// See `AutoRTFM::ForTheRuntime::CoinTossDisable` for implementation details.
+float GAutoRTFMEnabledProbability = 100.0f;
 
-	AUTORTFM_SANITIZER_ONLY(AutoRTFM::Sanitizer::EMode GSanitizerMode = AutoRTFM::Sanitizer::EMode::Error);
-	AUTORTFM_SANITIZER_ONLY(bool GSanitizerRecordClosedWriteCallstacks = false);
+AUTORTFM_SANITIZER_ONLY(AutoRTFM::Sanitizer::EMode GSanitizerMode = AutoRTFM::Sanitizer::EMode::Error);
+AUTORTFM_SANITIZER_ONLY(bool GSanitizerRecordClosedWriteCallstacks = false);
 
-	bool GAutoRTFMEnsureOnInternalAbort = true;
+bool GAutoRTFMEnsureOnInternalAbort = true;
 
-	// A linked-list of open->closed function tables populated by
-	// autortfm_register_open_to_closed_functions(). This is consumed by
-	// ProcessAllPendingOpenToClosedRegistrations() when autortfm_initialize()
-	// is called.
-	autortfm_open_to_closed_table* GPendingOpenToClosedRegistrations = nullptr;
+// A linked-list of open->closed function tables populated by
+// autortfm_register_open_to_closed_functions(). This is consumed by
+// ProcessAllPendingOpenToClosedRegistrations() when autortfm_initialize()
+// is called.
+autortfm_open_to_closed_table* GPendingOpenToClosedRegistrations = nullptr;
 
-	constexpr int GMaxRegisteredHeapRedirects = 4;
-	AutoRTFM::FHeapRedirectCallbacks GRegisteredHeapRedirects[GMaxRegisteredHeapRedirects];
-	int GNumRegisteredHeapRedirects = 0;
+constexpr int GMaxRegisteredHeapRedirects = 4;
+AutoRTFM::FHeapRedirectCallbacks GRegisteredHeapRedirects[GMaxRegisteredHeapRedirects];
+int GNumRegisteredHeapRedirects = 0;
 }
 #endif  // UE_AUTORTFM
 
 namespace AutoRTFM
 {
-	namespace Testing
+namespace Testing
+{
+UE_AUTORTFM_API ForTheRuntime::EAutoRTFMEnabledState ForceSetAutoRTFMRuntime(ForTheRuntime::EAutoRTFMEnabledState State)
+{
+#if UE_AUTORTFM
+	const int Original = GAutoRTFMRuntimeEnabled;
+	if (GAutoRTFMRuntimeEnabled != State)
 	{
-		UE_AUTORTFM_API ForTheRuntime::EAutoRTFMEnabledState ForceSetAutoRTFMRuntime(ForTheRuntime::EAutoRTFMEnabledState State)
+		GAutoRTFMRuntimeEnabled = State;
+		if (GExternAPI.OnRuntimeEnabledChanged)
 		{
-#if UE_AUTORTFM
-			const int Original = GAutoRTFMRuntimeEnabled;
-			if (GAutoRTFMRuntimeEnabled != State)
-			{
-				GAutoRTFMRuntimeEnabled = State;
-				if (GExternAPI.OnRuntimeEnabledChanged)
-				{
-					GExternAPI.OnRuntimeEnabledChanged();
-				}
-			}
-			return static_cast<ForTheRuntime::EAutoRTFMEnabledState>(Original);
+			GExternAPI.OnRuntimeEnabledChanged();
+		}
+	}
+	return static_cast<ForTheRuntime::EAutoRTFMEnabledState>(Original);
 #else
-			UE_AUTORTFM_UNUSED(State);
-			return ForTheRuntime::AutoRTFM_ForcedDisabled;
+	UE_AUTORTFM_UNUSED(State);
+	return ForTheRuntime::AutoRTFM_ForcedDisabled;
 #endif
-		}
-	} // Testing
-} // AutoRTFM
-
+}
+}  // namespace Testing
+}  // namespace AutoRTFM
 
 #if UE_AUTORTFM
 namespace AutoRTFM
 {
-	AUTORTFM_INTERNAL bool IsAutoRTFMInitialized()
-	{
-		return FContext::Get() != nullptr;
-	}
+AUTORTFM_INTERNAL bool IsAutoRTFMInitialized()
+{
+	return FContext::Get() != nullptr;
+}
 }
 #endif
 
@@ -87,79 +86,81 @@ namespace AutoRTFM
 {
 
 #if AUTORTFM_SANITIZER
-	namespace Sanitizer
+namespace Sanitizer
+{
+EMode GetMode()
+{
+	return GSanitizerMode;
+}
+
+void SetMode(EMode Mode)
+{
+	if (GSanitizerMode != Mode)
 	{
-		EMode GetMode()
+		GSanitizerMode = Mode;
+		if (GExternAPI.OnSanitizerModeChanged)
 		{
-			return GSanitizerMode;
+			GExternAPI.OnSanitizerModeChanged();
 		}
+	}
+}
 
-		void SetMode(EMode Mode)
+void BeginDisable()
+{
+	if (FContext* Context = FContext::Get())
+	{
+		if (FTransaction* Transaction = Context->GetCurrentTransaction())
 		{
-			if (GSanitizerMode != Mode)
-			{
-				GSanitizerMode = Mode;
-				if (GExternAPI.OnSanitizerModeChanged)
-				{
-					GExternAPI.OnSanitizerModeChanged();
-				}
-			}
+			Transaction->DisableSanitizer();
 		}
+	}
+}
 
-		void BeginDisable()
+void EndDisable()
+{
+	if (FContext* Context = FContext::Get())
+	{
+		if (FTransaction* Transaction = Context->GetCurrentTransaction())
 		{
-			if (FContext* Context = FContext::Get())
-			{
-				if (FTransaction* Transaction = Context->GetCurrentTransaction())
-				{
-					Transaction->DisableSanitizer();
-				}
-			}
+			Transaction->EnableSanitizer();
 		}
+	}
+}
 
-		void EndDisable()
-		{
-			if (FContext* Context = FContext::Get())
-			{
-				if (FTransaction* Transaction = Context->GetCurrentTransaction())
-				{
-					Transaction->EnableSanitizer();
-				}
-			}
-		}
+bool RecordClosedWriteCallstacks()
+{
+	return GSanitizerRecordClosedWriteCallstacks;
+}
 
-		bool RecordClosedWriteCallstacks()
+void SetRecordClosedWriteCallstacks(bool bEnable)
+{
+	if (GSanitizerRecordClosedWriteCallstacks != bEnable)
+	{
+		GSanitizerRecordClosedWriteCallstacks = bEnable;
+		if (GExternAPI.OnSanitizerRecordClosedWriteCallstacksChanged)
 		{
-			return GSanitizerRecordClosedWriteCallstacks;
-		}
-
-		void SetRecordClosedWriteCallstacks(bool bEnable)
-		{
-			if (GSanitizerRecordClosedWriteCallstacks != bEnable)
-			{
-				GSanitizerRecordClosedWriteCallstacks = bEnable;
-				if (GExternAPI.OnSanitizerRecordClosedWriteCallstacksChanged)
-				{
-					GExternAPI.OnSanitizerRecordClosedWriteCallstacksChanged();
-				}
-			}
+			GExternAPI.OnSanitizerRecordClosedWriteCallstacksChanged();
 		}
 	}
-#endif // AUTORTFM_SANITIZER
+}
+}  // namespace Sanitizer
+#endif  // AUTORTFM_SANITIZER
 
-	namespace ForTheRuntime
+namespace ForTheRuntime
+{
+bool SetAutoRTFMRuntime(EAutoRTFMEnabledState State)
+{
+	// #noop if AutoRTFM is not compiled in
+#if UE_AUTORTFM
+	auto Stringify = [](int Query) -> const char*
 	{
-		bool SetAutoRTFMRuntime(EAutoRTFMEnabledState State)
+		switch (Query)
 		{
-			// #noop if AutoRTFM is not compiled in
-#if UE_AUTORTFM
-			auto Stringify = [](int Query) -> const char*
-			{
-				switch (Query)
-				{
-				default:
-					InternalUnreachable();
-#define HANDLE_CASE(x) case x: return #x
+			default:
+				InternalUnreachable();
+#define HANDLE_CASE(x) \
+	case x:            \
+		return #x
 				HANDLE_CASE(EAutoRTFMEnabledState::AutoRTFM_ForcedEnabled);
 				HANDLE_CASE(EAutoRTFMEnabledState::AutoRTFM_ForcedDisabled);
 				HANDLE_CASE(EAutoRTFMEnabledState::AutoRTFM_OverriddenEnabled);
@@ -169,221 +170,221 @@ namespace AutoRTFM
 				HANDLE_CASE(EAutoRTFMEnabledState::AutoRTFM_EnabledByDefault);
 				HANDLE_CASE(EAutoRTFMEnabledState::AutoRTFM_DisabledByDefault);
 #undef HANDLE_CASE
-				}
-			};
+		}
+	};
 
-			auto DoIgnoreLog = [&](int State, int Stored)
-			{
-				AUTORTFM_LOG("Ignoring changing AutoRTFM runtime state to '%s' as it was previously set to '%s'", Stringify(State), Stringify(Stored));
-			};
+	auto DoIgnoreLog = [&](int State, int Stored)
+	{
+		AUTORTFM_LOG("Ignoring changing AutoRTFM runtime state to '%s' as it was previously set to '%s'", Stringify(State),
+			Stringify(Stored));
+	};
+
+	switch (GAutoRTFMRuntimeEnabled)
+	{
+		default:
+			break;
+		case EAutoRTFMEnabledState::AutoRTFM_ForcedEnabled:
+		case EAutoRTFMEnabledState::AutoRTFM_ForcedDisabled:
+			DoIgnoreLog(State, GAutoRTFMRuntimeEnabled);
+			return false;
+	}
 
-			switch (GAutoRTFMRuntimeEnabled)
+	switch (GAutoRTFMRuntimeEnabled)
+	{
+		default:
+			break;
+		case EAutoRTFMEnabledState::AutoRTFM_OverriddenEnabled:
+		case EAutoRTFMEnabledState::AutoRTFM_OverriddenDisabled:
+			if ((State == EAutoRTFMEnabledState::AutoRTFM_Enabled) || (State == EAutoRTFMEnabledState::AutoRTFM_Disabled)
+				|| (State == EAutoRTFMEnabledState::AutoRTFM_EnabledByDefault)
+				|| (State == EAutoRTFMEnabledState::AutoRTFM_DisabledByDefault))
 			{
-			default:
-				break;
-			case EAutoRTFMEnabledState::AutoRTFM_ForcedEnabled:
-			case EAutoRTFMEnabledState::AutoRTFM_ForcedDisabled:
 				DoIgnoreLog(State, GAutoRTFMRuntimeEnabled);
 				return false;
 			}
-
-			switch (GAutoRTFMRuntimeEnabled)
+		case EAutoRTFMEnabledState::AutoRTFM_Enabled:
+		case EAutoRTFMEnabledState::AutoRTFM_Disabled:
+			if ((State == EAutoRTFMEnabledState::AutoRTFM_EnabledByDefault)
+				|| (State == EAutoRTFMEnabledState::AutoRTFM_DisabledByDefault))
 			{
-			default:
-				break;
-			case EAutoRTFMEnabledState::AutoRTFM_OverriddenEnabled:
-			case EAutoRTFMEnabledState::AutoRTFM_OverriddenDisabled:
-				if ((State == EAutoRTFMEnabledState::AutoRTFM_Enabled) ||
-					(State == EAutoRTFMEnabledState::AutoRTFM_Disabled) ||
-					(State == EAutoRTFMEnabledState::AutoRTFM_EnabledByDefault) ||
-					(State == EAutoRTFMEnabledState::AutoRTFM_DisabledByDefault))
-				{
-					DoIgnoreLog(State, GAutoRTFMRuntimeEnabled);
-					return false;
-				}
-			case EAutoRTFMEnabledState::AutoRTFM_Enabled:
-			case EAutoRTFMEnabledState::AutoRTFM_Disabled:
-				if ((State == EAutoRTFMEnabledState::AutoRTFM_EnabledByDefault) ||
-					(State == EAutoRTFMEnabledState::AutoRTFM_DisabledByDefault))
-				{
-					DoIgnoreLog(State, GAutoRTFMRuntimeEnabled);
-					return false;
-				}
+				DoIgnoreLog(State, GAutoRTFMRuntimeEnabled);
+				return false;
 			}
+	}
 
-			if (GAutoRTFMRuntimeEnabled != State)
-			{
-				GAutoRTFMRuntimeEnabled = State;
-				if (GExternAPI.OnRuntimeEnabledChanged)
-				{
-					GExternAPI.OnRuntimeEnabledChanged();
-				}
-			}
+	if (GAutoRTFMRuntimeEnabled != State)
+	{
+		GAutoRTFMRuntimeEnabled = State;
+		if (GExternAPI.OnRuntimeEnabledChanged)
+		{
+			GExternAPI.OnRuntimeEnabledChanged();
+		}
+	}
 
-			return true;
+	return true;
 #else
-			UE_AUTORTFM_UNUSED(State);
-			return false;
+	UE_AUTORTFM_UNUSED(State);
+	return false;
 #endif
-		}
+}
 
-		bool IsAutoRTFMRuntimeEnabledInternal()
-		{
-			// #noop if AutoRTFM is not compiled in
+bool IsAutoRTFMRuntimeEnabledInternal()
+{
+	// #noop if AutoRTFM is not compiled in
 #if UE_AUTORTFM
-			switch (GAutoRTFMRuntimeEnabled)
-			{
-			default:
-				return false;
-			case EAutoRTFMEnabledState::AutoRTFM_Enabled:
-			case EAutoRTFMEnabledState::AutoRTFM_ForcedEnabled:
-			case EAutoRTFMEnabledState::AutoRTFM_OverriddenEnabled:
-			case EAutoRTFMEnabledState::AutoRTFM_EnabledByDefault:
-				return true;
-			}
-#else
+	switch (GAutoRTFMRuntimeEnabled)
+	{
+		default:
 			return false;
+		case EAutoRTFMEnabledState::AutoRTFM_Enabled:
+		case EAutoRTFMEnabledState::AutoRTFM_ForcedEnabled:
+		case EAutoRTFMEnabledState::AutoRTFM_OverriddenEnabled:
+		case EAutoRTFMEnabledState::AutoRTFM_EnabledByDefault:
+			return true;
+	}
+#else
+	return false;
 #endif
-		}
+}
 
-		void SetAutoRTFMEnabledProbability(float Chance)
-		{
+void SetAutoRTFMEnabledProbability(float Chance)
+{
 #if UE_AUTORTFM
-			AUTORTFM_ASSERT(Chance >= 0.0f && Chance <= 100.0f);
-			GAutoRTFMEnabledProbability = Chance;
+	AUTORTFM_ASSERT(Chance >= 0.0f && Chance <= 100.0f);
+	GAutoRTFMEnabledProbability = Chance;
 #else
-			UE_AUTORTFM_UNUSED(Chance);
+	UE_AUTORTFM_UNUSED(Chance);
 #endif
-		}
+}
 
-		float GetAutoRTFMEnabledProbability()
-		{
+float GetAutoRTFMEnabledProbability()
+{
 #if UE_AUTORTFM
-			return GAutoRTFMEnabledProbability;
+	return GAutoRTFMEnabledProbability;
 #else
-			return 0.0f;
+	return 0.0f;
 #endif
-		}
+}
 
-		bool CoinTossDisable()
-		{
+bool CoinTossDisable()
+{
 #if UE_AUTORTFM
-			if (!IsAutoRTFMRuntimeEnabled())
-			{
-				return false;
-			}
+	if (!IsAutoRTFMRuntimeEnabled())
+	{
+		return false;
+	}
 
-			static std::random_device Device;
-			static std::mt19937 Generator(Device());
-			static std::uniform_real_distribution<float> Distribution(0.0f, 100.0f);
+	static std::random_device Device;
+	static std::mt19937 Generator(Device());
+	static std::uniform_real_distribution<float> Distribution(0.0f, 100.0f);
 
-			// A value in the range [0..100), EG. inclusive of 0, exclusive of 100.
-			// So a `GAutoRTFMEnabledProbability` of 100 is always greater than the
-			// potential random range, and `GAutoRTFMEnabledProbability` of 0 is
-			// always less than or equal to the range.
-			const float Random = Distribution(Generator);
+	// A value in the range [0..100), EG. inclusive of 0, exclusive of 100.
+	// So a `GAutoRTFMEnabledProbability` of 100 is always greater than the
+	// potential random range, and `GAutoRTFMEnabledProbability` of 0 is
+	// always less than or equal to the range.
+	const float Random = Distribution(Generator);
 
-			if (GAutoRTFMEnabledProbability <= Random)
-			{
-				// If we have the runtime cvar set to `ForcedEnabled` then it'll ignore this call!
-				return AutoRTFM::ForTheRuntime::SetAutoRTFMRuntime(AutoRTFM::ForTheRuntime::AutoRTFM_ForcedDisabled);
-			}
+	if (GAutoRTFMEnabledProbability <= Random)
+	{
+		// If we have the runtime cvar set to `ForcedEnabled` then it'll ignore this call!
+		return AutoRTFM::ForTheRuntime::SetAutoRTFMRuntime(AutoRTFM::ForTheRuntime::AutoRTFM_ForcedDisabled);
+	}
 #endif
 
-			return false;
-		}
+	return false;
+}
 
-		void SetInternalAbortAction(EAutoRTFMInternalAbortActionState State)
-		{
+void SetInternalAbortAction(EAutoRTFMInternalAbortActionState State)
+{
 #if UE_AUTORTFM
-			GAutoRTFMInternalAbortAction = State;
+	GAutoRTFMInternalAbortAction = State;
 #else
-			UE_AUTORTFM_UNUSED(State);
+	UE_AUTORTFM_UNUSED(State);
 #endif
-		}
+}
 
-		EAutoRTFMInternalAbortActionState GetInternalAbortAction()
-		{
+EAutoRTFMInternalAbortActionState GetInternalAbortAction()
+{
 #if UE_AUTORTFM
-			return static_cast<EAutoRTFMInternalAbortActionState>(GAutoRTFMInternalAbortAction);
+	return static_cast<EAutoRTFMInternalAbortActionState>(GAutoRTFMInternalAbortAction);
 #else
-			return Crash;
+	return Crash;
 #endif
-		}
+}
 
-		bool GetEnsureOnInternalAbort()
-		{
+bool GetEnsureOnInternalAbort()
+{
 #if UE_AUTORTFM
-			return GAutoRTFMEnsureOnInternalAbort;
+	return GAutoRTFMEnsureOnInternalAbort;
 #else
-			return false;
+	return false;
 #endif
-		}
+}
 
-		void SetEnsureOnInternalAbort([[maybe_unused]] bool bEnabled)
-		{
+void SetEnsureOnInternalAbort([[maybe_unused]] bool bEnabled)
+{
 #if UE_AUTORTFM
-			GAutoRTFMEnsureOnInternalAbort = bEnabled;
+	GAutoRTFMEnsureOnInternalAbort = bEnabled;
 #endif
-		}
+}
 
-		void SetRetryTransaction(EAutoRTFMRetryTransactionState State)
-		{
+void SetRetryTransaction(EAutoRTFMRetryTransactionState State)
+{
 #if UE_AUTORTFM
-			if (GAutoRTFMRetryTransactions != State)
-			{
-				GAutoRTFMRetryTransactions = State;
-				if (GExternAPI.OnRetryTransactionsChanged)
-				{
-					GExternAPI.OnRetryTransactionsChanged();
-				}
-			}
+	if (GAutoRTFMRetryTransactions != State)
+	{
+		GAutoRTFMRetryTransactions = State;
+		if (GExternAPI.OnRetryTransactionsChanged)
+		{
+			GExternAPI.OnRetryTransactionsChanged();
+		}
+	}
 #else
-			UE_AUTORTFM_UNUSED(State);
+	UE_AUTORTFM_UNUSED(State);
 #endif
-		}
+}
 
-		EAutoRTFMRetryTransactionState GetRetryTransaction()
-		{
+EAutoRTFMRetryTransactionState GetRetryTransaction()
+{
 #if UE_AUTORTFM
-			return static_cast<EAutoRTFMRetryTransactionState>(GAutoRTFMRetryTransactions);
+	return static_cast<EAutoRTFMRetryTransactionState>(GAutoRTFMRetryTransactions);
 #else
-			return NoRetry;
+	return NoRetry;
 #endif
-		}
+}
 
-		bool ShouldRetryNonNestedTransactions()
-		{
+bool ShouldRetryNonNestedTransactions()
+{
 #if UE_AUTORTFM
-			switch (GAutoRTFMRetryTransactions)
-			{
-			default:
-				return false;
-			case EAutoRTFMRetryTransactionState::RetryNonNested:
-			case EAutoRTFMRetryTransactionState::RetryNestedToo:
-				return true;
-			}
-#else
+	switch (GAutoRTFMRetryTransactions)
+	{
+		default:
 			return false;
+		case EAutoRTFMRetryTransactionState::RetryNonNested:
+		case EAutoRTFMRetryTransactionState::RetryNestedToo:
+			return true;
+	}
+#else
+	return false;
 #endif
-		}
+}
 
-		bool ShouldRetryNestedTransactionsToo()
-		{
+bool ShouldRetryNestedTransactionsToo()
+{
 #if UE_AUTORTFM
-			switch (GAutoRTFMRetryTransactions)
-			{
-			default:
-				return false;
-			case EAutoRTFMRetryTransactionState::RetryNestedToo:
-				return true;
-			}
-#else
+	switch (GAutoRTFMRetryTransactions)
+	{
+		default:
 			return false;
-#endif
-		}
+		case EAutoRTFMRetryTransactionState::RetryNestedToo:
+			return true;
 	}
+#else
+	return false;
+#endif
 }
+}  // namespace ForTheRuntime
+}  // namespace AutoRTFM
 
 #if (defined(__AUTORTFM_ENABLED) && __AUTORTFM_ENABLED)
 #include "AutoRTFM/Constants.h"
@@ -392,8 +393,8 @@ namespace AutoRTFM
 #include "ContextInlines.h"
 #include "ContextStatus.h"
 #include "FunctionMapInlines.h"
-#include "TransactionInlines.h"
 #include "Toggles.h"
+#include "TransactionInlines.h"
 #include "Utils.h"
 
 // This is the implementation of the AutoRTFM.h API. Ideally, functions here should just delegate to some internal API.
@@ -411,7 +412,8 @@ AUTORTFM_INTERNAL bool RTFM_autortfm_is_transactional()
 	return true;
 }
 
-AUTORTFM_INTERNAL autortfm_result RTFM_autortfm_transact(void (*UninstrumentedWork)(void*), void (*InstrumentedWork)(void*), void* Arg)
+AUTORTFM_INTERNAL autortfm_result RTFM_autortfm_transact(
+	void (*UninstrumentedWork)(void*), void (*InstrumentedWork)(void*), void* Arg)
 {
 	FContext* Context = FContext::Get();
 	return static_cast<autortfm_result>(Context->Transact(UninstrumentedWork, InstrumentedWork, Arg));
@@ -419,11 +421,7 @@ AUTORTFM_INTERNAL autortfm_result RTFM_autortfm_transact(void (*UninstrumentedWo
 
 UE_AUTORTFM_FORCEINLINE autortfm_result TransactThenOpenImpl(void (*UninstrumentedWork)(void*), void* Arg)
 {
-	return static_cast<autortfm_result>(
-		AutoRTFM::Transact([&]
-		{
-			autortfm_open(UninstrumentedWork, Arg);
-		}));
+	return static_cast<autortfm_result>(AutoRTFM::Transact([&] { autortfm_open(UninstrumentedWork, Arg); }));
 }
 
 AUTORTFM_INTERNAL autortfm_result RTFM_autortfm_transact_then_open(void (*UninstrumentedWork)(void*), void* Arg)
@@ -453,8 +451,8 @@ AUTORTFM_INTERNAL autortfm_status RTFM_autortfm_close(void (*UninstrumentedWork)
 	}
 	else
 	{
-		AUTORTFM_REPORT_ERROR("Could not find function %p '%s' for autortfm_close()",
-			reinterpret_cast<void*>(UninstrumentedWork), GetFunctionDescription(UninstrumentedWork).c_str());
+		AUTORTFM_REPORT_ERROR("Could not find function %p '%s' for autortfm_close()", reinterpret_cast<void*>(UninstrumentedWork),
+			GetFunctionDescription(UninstrumentedWork).c_str());
 	}
 
 	return static_cast<autortfm_status>(Context->GetStatus());
@@ -467,7 +465,7 @@ extern "C" UE_AUTORTFM_NOAUTORTFM UE_AUTORTFM_API void autortfm_pre_open(bool bN
 	{
 		FContext::Get()->GetCurrentTransaction()->DisableSanitizer();
 	}
-#endif // AUTORTFM_SANITIZER
+#endif  // AUTORTFM_SANITIZER
 }
 
 extern "C" UE_AUTORTFM_NOAUTORTFM UE_AUTORTFM_API void autortfm_post_open(bool bNoSanitize)
@@ -479,7 +477,7 @@ extern "C" UE_AUTORTFM_NOAUTORTFM UE_AUTORTFM_API void autortfm_post_open(bool b
 	{
 		FContext::Get()->GetCurrentTransaction()->EnableSanitizer();
 	}
-#endif // AUTORTFM_SANITIZER
+#endif  // AUTORTFM_SANITIZER
 
 	if (Context->IsAborting())
 	{
@@ -519,7 +517,7 @@ extern "C" AUTORTFM_INTERNAL UE_AUTORTFM_API void autortfm_open_from_closed_no_s
 	FContext* Context = FContext::Get();
 
 	FTransaction* Transaction = Context->GetCurrentTransaction();
-	
+
 	AUTORTFM_SANITIZER_ONLY(Transaction->DisableSanitizer());
 
 	Work(Arg);
@@ -553,7 +551,7 @@ AUTORTFM_INTERNAL void RTFM_CascadingRetryTransactionInternal()
 	Context->AbortTransaction(EContextStatus::AbortedByCascadingRetry);
 }
 
-AUTORTFM_INTERNAL void RTFM_OnCommitInternal(TTask<void()> && Work)
+AUTORTFM_INTERNAL void RTFM_OnCommitInternal(TTask<void()>&& Work)
 {
 	FContext* Context = FContext::Get();
 	AUTORTFM_ASSERT(Context->GetStatus() == EContextStatus::OnTrack);
@@ -777,7 +775,7 @@ extern "C" autortfm_result autortfm_transact(void (*UninstrumentedWork)(void*),
 
 	if (ForTheRuntime::IsAutoRTFMRuntimeEnabled())
 	{
-	    return static_cast<autortfm_result>(FContext::Get()->Transact(UninstrumentedWork, InstrumentedWork, Arg));
+		return static_cast<autortfm_result>(FContext::Get()->Transact(UninstrumentedWork, InstrumentedWork, Arg));
 	}
 
 	(*UninstrumentedWork)(Arg);
@@ -813,7 +811,7 @@ extern "C" autortfm_transaction_id autortfm_current_transaction_id() noexcept
 
 extern "C" void autortfm_abort_transaction() noexcept
 {
-	AUTORTFM_ASSERT(autortfm_is_closed()); // RollbackTransaction should be used when in the open.
+	AUTORTFM_ASSERT(autortfm_is_closed());  // RollbackTransaction should be used when in the open.
 	RTFM_autortfm_abort_transaction();
 }
 
@@ -869,13 +867,15 @@ extern "C" autortfm_status autortfm_close(void (*UninstrumentedWork)(void*), voi
 			std::string FunctionDescription = GetFunctionDescription(UninstrumentedWork);
 			if (ForTheRuntime::GetInternalAbortAction() == ForTheRuntime::EAutoRTFMInternalAbortActionState::Crash)
 			{
-				AUTORTFM_FATAL("Could not find function %p '%s' in autortfm_close()", UninstrumentedWork, FunctionDescription.c_str());
+				AUTORTFM_FATAL(
+					"Could not find function %p '%s' in autortfm_close()", UninstrumentedWork, FunctionDescription.c_str());
 			}
 			else
 			{
-				AUTORTFM_ENSURE_MSG(!ForTheRuntime::GetEnsureOnInternalAbort(), "Could not find function %p '%s' in autortfm_close()", UninstrumentedWork, FunctionDescription.c_str());
+				AUTORTFM_ENSURE_MSG(!ForTheRuntime::GetEnsureOnInternalAbort(),
+					"Could not find function %p '%s' in autortfm_close()", UninstrumentedWork, FunctionDescription.c_str());
 			}
-	        Context->AbortByLanguageAndThrow();
+			Context->AbortByLanguageAndThrow();
 		}
 	}
 	else
@@ -1032,7 +1032,7 @@ void ForTheRuntime::CascadingRetryRollbackTransaction()
 	Context->RollbackTransaction(AutoRTFM::EContextStatus::AbortedByCascadingRetry);
 }
 
-void ForTheRuntime::OnCommitInternal(TTask<void()> && Work)
+void ForTheRuntime::OnCommitInternal(TTask<void()>&& Work)
 {
 	if (autortfm_is_closed())
 	{
@@ -1042,14 +1042,14 @@ void ForTheRuntime::OnCommitInternal(TTask<void()> && Work)
 	Work();
 }
 
-void ForTheRuntime::OnPreAbortInternal(TTask<void()> && Work)
+void ForTheRuntime::OnPreAbortInternal(TTask<void()>&& Work)
 {
 	if (autortfm_is_closed())
 	{
 		return RTFM_OnPreAbortInternal(std::move(Work));
 	}
 }
-void ForTheRuntime::OnAbortInternal(TTask<void()> && Work)
+void ForTheRuntime::OnAbortInternal(TTask<void()>&& Work)
 {
 	if (autortfm_is_closed())
 	{
@@ -1057,7 +1057,7 @@ void ForTheRuntime::OnAbortInternal(TTask<void()> && Work)
 	}
 }
 
-void ForTheRuntime::OnCompleteInternal(TTask<void()> && Work)
+void ForTheRuntime::OnCompleteInternal(TTask<void()>&& Work)
 {
 	if (autortfm_is_closed())
 	{
@@ -1089,7 +1089,7 @@ void ForTheRuntime::PopAllOnCommitHandlersInternal(const void* Key)
 	}
 }
 
-void ForTheRuntime::PushOnAbortHandlerInternal(const void* Key, TTask<void()> && Work)
+void ForTheRuntime::PushOnAbortHandlerInternal(const void* Key, TTask<void()>&& Work)
 {
 	if (autortfm_is_closed())
 	{
@@ -1114,20 +1114,20 @@ void ForTheRuntime::PopAllOnAbortHandlersInternal(const void* Key)
 }
 
 UE_AUTORTFM_NOAUTORTFM
-void ForTheRuntime::RegisterOnCommitFromTheOpen(TTask<void()> && Work)
+void ForTheRuntime::RegisterOnCommitFromTheOpen(TTask<void()>&& Work)
 {
 	RTFM_OnCommitInternal(std::move(Work));
 }
 
 UE_AUTORTFM_NOAUTORTFM
-void ForTheRuntime::RegisterOnAbortFromTheOpen(TTask<void()> && Work)
+void ForTheRuntime::RegisterOnAbortFromTheOpen(TTask<void()>&& Work)
 {
 	RTFM_OnAbortInternal(std::move(Work));
 }
 
 void ForTheRuntime::RedirectedLoad(uint32_t AddressSpace, void* DestPointer, uint64_t Size, uint64_t SourceAddress)
 {
-	for(int RedirectedHeapIndex = 0; RedirectedHeapIndex < GNumRegisteredHeapRedirects; RedirectedHeapIndex++)
+	for (int RedirectedHeapIndex = 0; RedirectedHeapIndex < GNumRegisteredHeapRedirects; RedirectedHeapIndex++)
 	{
 		FHeapRedirectCallbacks* HeapRedirects = &GRegisteredHeapRedirects[RedirectedHeapIndex];
 
@@ -1141,7 +1141,7 @@ void ForTheRuntime::RedirectedLoad(uint32_t AddressSpace, void* DestPointer, uin
 
 void ForTheRuntime::RedirectedStore(uint32_t AddressSpace, uint64_t DestAddress, uint64_t Size, const void* SourcePointer)
 {
-	for(int RedirectedHeapIndex = 0; RedirectedHeapIndex < GNumRegisteredHeapRedirects; RedirectedHeapIndex++)
+	for (int RedirectedHeapIndex = 0; RedirectedHeapIndex < GNumRegisteredHeapRedirects; RedirectedHeapIndex++)
 	{
 		FHeapRedirectCallbacks* HeapRedirects = &GRegisteredHeapRedirects[RedirectedHeapIndex];
 
@@ -1172,8 +1172,10 @@ void RegisterHeapRedirectCallbacks(FHeapRedirectCallbacks Callbacks)
 	// function we translate the function pointers passed in to be their closed variants
 	FHeapRedirectCallbacks ClosedCallbacks;
 	ClosedCallbacks.AddressSpace = Callbacks.AddressSpace;
-	ClosedCallbacks.RedirectedLoad = reinterpret_cast<void (*)(void*, uint64_t, uint64_t)>(autortfm_lookup_function(reinterpret_cast<void*>(Callbacks.RedirectedLoad), "RegisterHeapRedirects"));
-	ClosedCallbacks.RedirectedStore = reinterpret_cast<void (*)(uint64_t, uint64_t, const void*)>(autortfm_lookup_function(reinterpret_cast<void*>(Callbacks.RedirectedStore), "RegisterHeapRedirects"));
+	ClosedCallbacks.RedirectedLoad = reinterpret_cast<void (*)(void*, uint64_t, uint64_t)>(
+		autortfm_lookup_function(reinterpret_cast<void*>(Callbacks.RedirectedLoad), "RegisterHeapRedirects"));
+	ClosedCallbacks.RedirectedStore = reinterpret_cast<void (*)(uint64_t, uint64_t, const void*)>(
+		autortfm_lookup_function(reinterpret_cast<void*>(Callbacks.RedirectedStore), "RegisterHeapRedirects"));
 	GRegisteredHeapRedirects[GNumRegisteredHeapRedirects++] = ClosedCallbacks;
 }
 
@@ -1184,7 +1186,7 @@ extern "C" void autortfm_on_commit(void (*Work)(void*), void* Arg)
 		return RTFM_autortfm_on_commit(Work, Arg);
 	}
 
-    Work(Arg);
+	Work(Arg);
 }
 
 extern "C" void autortfm_on_pre_abort(void (*Work)(void*), void* Arg) noexcept
@@ -1201,7 +1203,6 @@ extern "C" void autortfm_on_abort(void (*Work)(void*), void* Arg) noexcept
 	{
 		return RTFM_autortfm_on_abort(Work, Arg);
 	}
-
 }
 
 extern "C" void autortfm_push_on_abort_handler(const void* Key, void (*Work)(void*), void* Arg) noexcept
@@ -1210,7 +1211,6 @@ extern "C" void autortfm_push_on_abort_handler(const void* Key, void (*Work)(voi
 	{
 		return RTFM_autortfm_push_on_abort_handler(Key, Work, Arg);
 	}
-
 }
 
 extern "C" void autortfm_pop_on_abort_handler(const void* Key) noexcept
@@ -1219,7 +1219,6 @@ extern "C" void autortfm_pop_on_abort_handler(const void* Key) noexcept
 	{
 		return RTFM_autortfm_pop_on_abort_handler(Key);
 	}
-
 }
 
 extern "C" void* autortfm_did_allocate(void* Ptr, size_t Size) noexcept
@@ -1229,7 +1228,7 @@ extern "C" void* autortfm_did_allocate(void* Ptr, size_t Size) noexcept
 		return RTFM_autortfm_did_allocate(Ptr, Size);
 	}
 
-    return Ptr;
+	return Ptr;
 }
 
 extern "C" void autortfm_did_free(void* Ptr) noexcept
@@ -1262,8 +1261,8 @@ extern "C" void autortfm_did_free(void* Ptr) noexcept
 // code.
 extern "C" UE_AUTORTFM_NOAUTORTFM void autortfm_check_abi(void* const Ptr, const size_t Size) noexcept
 {
-    struct FConstants final
-    {
+	struct FConstants final
+	{
 		const uint32_t Major = AutoRTFM::Constants::Major;
 		const uint32_t Minor = AutoRTFM::Constants::Minor;
 		const uint32_t Patch = AutoRTFM::Constants::Patch;
@@ -1287,14 +1286,14 @@ extern "C" UE_AUTORTFM_NOAUTORTFM void autortfm_check_abi(void* const Ptr, const
 		{
 			return !(*this == Other);
 		}
-    } RuntimeConstants;
+	} RuntimeConstants;
 
 	AUTORTFM_FATAL_IF(sizeof(FConstants) != Size, "ABI error between AutoRTFM compiler and runtime");
 
-    const FConstants* const CompilerConstants = static_cast<FConstants*>(Ptr);
+	const FConstants* const CompilerConstants = static_cast<FConstants*>(Ptr);
 
 	AUTORTFM_FATAL_IF(RuntimeConstants != *CompilerConstants, "ABI error between AutoRTFM compiler and runtime");
 }
-} // namespace AutoRTFM
+}  // namespace AutoRTFM
 
-#endif // defined(__AUTORTFM) && __AUTORTFM
+#endif  // defined(__AUTORTFM) && __AUTORTFM
diff --git a/Engine/Source/Runtime/AutoRTFM/Private/BlockAllocator.h b/Engine/Source/Runtime/AutoRTFM/Private/BlockAllocator.h
--- a/Engine/Source/Runtime/AutoRTFM/Private/BlockAllocator.h
+++ b/Engine/Source/Runtime/AutoRTFM/Private/BlockAllocator.h
@@ -4,8 +4,8 @@
 
 #if (defined(__AUTORTFM) && __AUTORTFM)
 
-#include "Utils.h"
 #include "ExternAPI.h"
+#include "Utils.h"
 
 #include <cstddef>
 
@@ -21,7 +21,7 @@ namespace AutoRTFM
 //   InlineBlockDataSize - The size of the data section of the inline block.
 //   DataAlignment - The maximum supported alignment of data from this allocator.
 //   GrowthPercentage - The percentage to grow the block size each reallocation.
-template<size_t InlineBlockDataSize = 256, size_t DataAlignment = 16, size_t GrowthPercentage = 200>
+template <size_t InlineBlockDataSize = 256, size_t DataAlignment = 16, size_t GrowthPercentage = 200>
 class TBlockAllocator
 {
 public:
@@ -70,7 +70,8 @@ public:
 
 	// Constructs and returns a new T into the memory returned by calling
 	// Allocate(sizeof(T), alignof(T)).
-	template<typename T, typename... ArgTypes> T* New(ArgTypes&&... Args)
+	template <typename T, typename... ArgTypes>
+	T* New(ArgTypes&&... Args)
 	{
 		return new (Allocate(sizeof(T), alignof(T))) T(Forward<ArgTypes>(Args)...);
 	}
@@ -80,7 +81,7 @@ private:
 	TBlockAllocator(const TBlockAllocator&) = delete;
 	TBlockAllocator& operator=(const TBlockAllocator&) = delete;
 	TBlockAllocator& operator=(TBlockAllocator&&) = delete;
-	
+
 	static constexpr size_t BlockAlignment = std::max<size_t>(DataAlignment, 8);
 
 	struct alignas(BlockAlignment) FBlockHeader
@@ -123,6 +124,6 @@ private:
 	alignas(BlockAlignment) std::byte Data[InlineBlockDataSize];
 };
 
-}
+}  // namespace AutoRTFM
 
-#endif // (defined(__AUTORTFM) && __AUTORTFM)
+#endif  // (defined(__AUTORTFM) && __AUTORTFM)
diff --git a/Engine/Source/Runtime/AutoRTFM/Private/CompilerInterface.cpp b/Engine/Source/Runtime/AutoRTFM/Private/CompilerInterface.cpp
--- a/Engine/Source/Runtime/AutoRTFM/Private/CompilerInterface.cpp
+++ b/Engine/Source/Runtime/AutoRTFM/Private/CompilerInterface.cpp
@@ -18,21 +18,19 @@ namespace AutoRTFM
 // Check for writes to null in development code, so that the inevitable crash will occur
 // in the caller's code rather than in the AutoRTFM runtime.
 #define UE_AUTORTFM_HANDLE_NULL_WRITE(Ptr)     \
-    do                                         \
-    {                                          \
-        if (AUTORTFM_UNLIKELY(Ptr == nullptr)) \
-        {                                      \
-            return;                            \
-        }                                      \
-    }                                          \
-    while (0)
+	do                                         \
+	{                                          \
+		if (AUTORTFM_UNLIKELY(Ptr == nullptr)) \
+		{                                      \
+			return;                            \
+		}                                      \
+	} while (0)
 
 #else
 
 // In shipping code, we don't want to spend any cycles on a redundant check.
 // We do want the compiler to optimize as if the pointer is non-null, though.
-#define UE_AUTORTFM_HANDLE_NULL_WRITE(Ptr)  \
-    UE_ASSUME(Ptr != nullptr)
+#define UE_AUTORTFM_HANDLE_NULL_WRITE(Ptr) UE_ASSUME(Ptr != nullptr)
 
 #endif
 
@@ -71,12 +69,13 @@ extern "C" AUTORTFM_DISABLE_NO_SANITIZE UE_AUTORTFM_API AUTORTFM_NO_ASAN void au
 	Context->RecordWrite<8>(Ptr);
 }
 
-extern "C" AUTORTFM_DISABLE_NO_SANITIZE UE_AUTORTFM_API AUTORTFM_NO_ASAN void autortfm_record_masked_write(void* Ptr, uintptr_t Mask, int NumMaskBits, int ValueSizeBytes) noexcept
+extern "C" AUTORTFM_DISABLE_NO_SANITIZE UE_AUTORTFM_API AUTORTFM_NO_ASAN void autortfm_record_masked_write(
+	void* Ptr, uintptr_t Mask, int NumMaskBits, int ValueSizeBytes) noexcept
 {
 	UE_AUTORTFM_HANDLE_NULL_WRITE(Ptr);
 
 	char* IncrementablePtr = static_cast<char*>(Ptr);
-	for(int i = 0; i < NumMaskBits; i++)
+	for (int i = 0; i < NumMaskBits; i++)
 	{
 		if (Mask & (1u << i))
 		{
@@ -87,7 +86,8 @@ extern "C" AUTORTFM_DISABLE_NO_SANITIZE UE_AUTORTFM_API AUTORTFM_NO_ASAN void au
 	}
 }
 
-extern "C" UE_AUTORTFM_API AUTORTFM_NO_ASAN void autortfm_redirected_load(uint32 AddressSpace, void* DestPointer, uint64 Size, uint64 SourceAddress) noexcept
+extern "C" UE_AUTORTFM_API AUTORTFM_NO_ASAN void autortfm_redirected_load(
+	uint32 AddressSpace, void* DestPointer, uint64 Size, uint64 SourceAddress) noexcept
 {
 	ForTheRuntime::RedirectedLoad(AddressSpace, DestPointer, Size, SourceAddress);
 }
@@ -120,27 +120,32 @@ extern "C" UE_AUTORTFM_API AUTORTFM_NO_ASAN uint8 autortfm_redirected_load_1(uin
 	return Result;
 }
 
-extern "C" UE_AUTORTFM_API AUTORTFM_NO_ASAN void autortfm_redirected_store(uint32 AddressSpace, uint64 DestAddress, uint64 Size, const void* SourcePointer) noexcept
+extern "C" UE_AUTORTFM_API AUTORTFM_NO_ASAN void autortfm_redirected_store(
+	uint32 AddressSpace, uint64 DestAddress, uint64 Size, const void* SourcePointer) noexcept
 {
 	ForTheRuntime::RedirectedStore(AddressSpace, DestAddress, Size, SourcePointer);
 }
 
-extern "C" UE_AUTORTFM_API AUTORTFM_NO_ASAN void autortfm_redirected_store_8(uint32 AddressSpace, uint64 DestAddress, uint64 Value) noexcept
+extern "C" UE_AUTORTFM_API AUTORTFM_NO_ASAN void autortfm_redirected_store_8(
+	uint32 AddressSpace, uint64 DestAddress, uint64 Value) noexcept
 {
 	ForTheRuntime::RedirectedStore(AddressSpace, DestAddress, sizeof(Value), &Value);
 }
 
-extern "C" UE_AUTORTFM_API AUTORTFM_NO_ASAN void autortfm_redirected_store_4(uint32 AddressSpace, uint64 DestAddress, uint32 Value) noexcept
+extern "C" UE_AUTORTFM_API AUTORTFM_NO_ASAN void autortfm_redirected_store_4(
+	uint32 AddressSpace, uint64 DestAddress, uint32 Value) noexcept
 {
 	ForTheRuntime::RedirectedStore(AddressSpace, DestAddress, sizeof(Value), &Value);
 }
 
-extern "C" UE_AUTORTFM_API AUTORTFM_NO_ASAN void autortfm_redirected_store_2(uint32 AddressSpace, uint64 DestAddress, uint16 Value) noexcept
+extern "C" UE_AUTORTFM_API AUTORTFM_NO_ASAN void autortfm_redirected_store_2(
+	uint32 AddressSpace, uint64 DestAddress, uint16 Value) noexcept
 {
 	ForTheRuntime::RedirectedStore(AddressSpace, DestAddress, sizeof(Value), &Value);
 }
 
-extern "C" UE_AUTORTFM_API AUTORTFM_NO_ASAN void autortfm_redirected_store_1(uint32 AddressSpace, uint64 DestAddress, uint8 Value) noexcept
+extern "C" UE_AUTORTFM_API AUTORTFM_NO_ASAN void autortfm_redirected_store_1(
+	uint32 AddressSpace, uint64 DestAddress, uint8 Value) noexcept
 {
 	ForTheRuntime::RedirectedStore(AddressSpace, DestAddress, sizeof(Value), &Value);
 }
@@ -154,19 +159,19 @@ extern "C" void* autortfm_lookup_function(void* OriginalFunction, const char* Wh
 extern "C" AUTORTFM_DISABLE_NO_SANITIZE UE_AUTORTFM_API void autortfm_memcpy(void* Dst, const void* Src, size_t Size) noexcept
 {
 	FContext* Context = FContext::Get();
-    Memcpy(Dst, Src, Size, Context);
+	Memcpy(Dst, Src, Size, Context);
 }
 
 extern "C" AUTORTFM_DISABLE_NO_SANITIZE UE_AUTORTFM_API void autortfm_memmove(void* Dst, const void* Src, size_t Size) noexcept
 {
 	FContext* Context = FContext::Get();
-    Memmove(Dst, Src, Size, Context);
+	Memmove(Dst, Src, Size, Context);
 }
 
 extern "C" AUTORTFM_DISABLE_NO_SANITIZE UE_AUTORTFM_API void autortfm_memset(void* Dst, int Value, size_t Size) noexcept
 {
 	FContext* Context = FContext::Get();
-    Memset(Dst, Value, Size, Context);
+	Memset(Dst, Value, Size, Context);
 }
 
 extern "C" UE_AUTORTFM_ALWAYS_OPEN UE_AUTORTFM_API void autortfm_unreachable(const char* Message) noexcept
@@ -199,7 +204,7 @@ extern "C" AUTORTFM_DISABLE_NO_SANITIZE UE_AUTORTFM_API void autortfm_llvm_missi
 	}
 
 	FContext* Context = FContext::Get();
-    Context->AbortByLanguageAndThrow();
+	Context->AbortByLanguageAndThrow();
 }
 
 extern "C" AUTORTFM_DISABLE_NO_SANITIZE UE_AUTORTFM_FORCENOINLINE UE_AUTORTFM_API void autortfm_called_no_autortfm() noexcept
@@ -214,7 +219,8 @@ extern "C" AUTORTFM_DISABLE_NO_SANITIZE UE_AUTORTFM_FORCENOINLINE UE_AUTORTFM_AP
 
 #if AUTORTFM_SANITIZER
 static bool bInAutoRTFMSanitizerOpenWrite = false;
-extern "C" AUTORTFM_DISABLE_NO_SANITIZE UE_AUTORTFM_FORCENOINLINE UE_AUTORTFM_API void autortfm_sanitizer_open_write(void* Ptr, size_t Size) noexcept
+extern "C" AUTORTFM_DISABLE_NO_SANITIZE UE_AUTORTFM_FORCENOINLINE UE_AUTORTFM_API void autortfm_sanitizer_open_write(
+	void* Ptr, size_t Size) noexcept
 {
 	if (FContext* Context = FContext::Get(); Context && Context->IsContextThread())
 	{
@@ -230,6 +236,6 @@ extern "C" AUTORTFM_DISABLE_NO_SANITIZE UE_AUTORTFM_FORCENOINLINE UE_AUTORTFM_AP
 }
 #endif
 
-} // namespace AutoRTFM
+}  // namespace AutoRTFM
 
-#endif // defined(__AUTORTFM) && __AUTORTFM
+#endif  // defined(__AUTORTFM) && __AUTORTFM
diff --git a/Engine/Source/Runtime/AutoRTFM/Private/ExternAPI.cpp b/Engine/Source/Runtime/AutoRTFM/Private/ExternAPI.cpp
--- a/Engine/Source/Runtime/AutoRTFM/Private/ExternAPI.cpp
+++ b/Engine/Source/Runtime/AutoRTFM/Private/ExternAPI.cpp
@@ -20,7 +20,8 @@ namespace
 // which have logic that can (potentially) log. Logging is unlikely to be called
 // before initialization, but a default printf implementation is preferable to
 // crashing due to a nullptr.
-void DefaultLog(const char* File, int Line, [[maybe_unused]] void* ProgramCounter, autortfm_log_severity Severity, const char* Format, va_list Args)
+void DefaultLog(const char* File, int Line, [[maybe_unused]] void* ProgramCounter, autortfm_log_severity Severity,
+	const char* Format, va_list Args)
 {
 	static constexpr int InlineBufferLength = 256;
 	char InlineBuffer[InlineBufferLength];
@@ -28,10 +29,10 @@ void DefaultLog(const char* File, int Line, [[maybe_unused]] void* ProgramCounte
 	va_list Args2;
 	va_copy(Args2, Args);
 	int Count = vsnprintf(InlineBuffer, InlineBufferLength, Format, Args);
-	auto AllocatedBuffer = (Count >= InlineBufferLength) ? std::unique_ptr<char[]>{new char[Count+1]} : nullptr;
+	auto AllocatedBuffer = (Count >= InlineBufferLength) ? std::unique_ptr<char[]>{new char[Count + 1]} : nullptr;
 	if (AllocatedBuffer)
 	{
-		vsnprintf(AllocatedBuffer.get(), Count+1, Format, Args2);
+		vsnprintf(AllocatedBuffer.get(), Count + 1, Format, Args2);
 	}
 	va_end(Args2);
 
@@ -39,29 +40,28 @@ void DefaultLog(const char* File, int Line, [[maybe_unused]] void* ProgramCounte
 
 	switch (Severity)
 	{
-	case autortfm_log_verbose:
-		fprintf(stdout, "AutoRTFM %s:%d [VERBOSE]: %s", File, Line, Buffer);
-		break;
-	case autortfm_log_info:
-		fprintf(stdout, "AutoRTFM %s:%d [INFO]: %s", File, Line, Buffer);
-		break;
-	case autortfm_log_warn:
-		fprintf(stdout, "AutoRTFM %s:%d [WARN]: %s", File, Line, Buffer);
-		break;
-	case autortfm_log_error:
-		fprintf(stderr, "AutoRTFM %s:%d [ERROR]: %s", File, Line, Buffer);
-		break;
-	case autortfm_log_fatal:
-		fprintf(stderr, "AutoRTFM %s:%d [FATAL]: %s", File, Line, Buffer);
-		__builtin_trap();
-		break;
+		case autortfm_log_verbose:
+			fprintf(stdout, "AutoRTFM %s:%d [VERBOSE]: %s", File, Line, Buffer);
+			break;
+		case autortfm_log_info:
+			fprintf(stdout, "AutoRTFM %s:%d [INFO]: %s", File, Line, Buffer);
+			break;
+		case autortfm_log_warn:
+			fprintf(stdout, "AutoRTFM %s:%d [WARN]: %s", File, Line, Buffer);
+			break;
+		case autortfm_log_error:
+			fprintf(stderr, "AutoRTFM %s:%d [ERROR]: %s", File, Line, Buffer);
+			break;
+		case autortfm_log_fatal:
+			fprintf(stderr, "AutoRTFM %s:%d [FATAL]: %s", File, Line, Buffer);
+			__builtin_trap();
+			break;
 	}
 }
 
 }  // anonymous namespace
 
-autortfm_extern_api GExternAPI
-{
+autortfm_extern_api GExternAPI = {
 	/* Allocate */ nullptr,
 	/* Reallocate */ nullptr,
 	/* AllocateZeroed */ nullptr,
@@ -81,6 +81,6 @@ autortfm_extern_api GExternAPI
 	/* RollbackWrite */ nullptr,
 };
 
-} // namespace AutoRTFM
+}  // namespace AutoRTFM
 
-#endif // defined(__AUTORTFM) && __AUTORTFM
+#endif  // defined(__AUTORTFM) && __AUTORTFM
diff --git a/Engine/Source/Runtime/AutoRTFM/Private/HashMap.h b/Engine/Source/Runtime/AutoRTFM/Private/HashMap.h
--- a/Engine/Source/Runtime/AutoRTFM/Private/HashMap.h
+++ b/Engine/Source/Runtime/AutoRTFM/Private/HashMap.h
@@ -14,7 +14,7 @@ namespace AutoRTFM
 {
 
 // A simple structure holding a templated key and value pair.
-template<typename KeyType, typename ValueType>
+template <typename KeyType, typename ValueType>
 struct TKeyAndValue
 {
 	KeyType Key;
@@ -24,16 +24,18 @@ struct TKeyAndValue
 // HashMap is an unordered hashmap.
 // TODO(SOL-7652): This currently wraps a std::unordered_map. Reimplement to
 // improve performance and to avoid platform-specific variations in behavior.
-template<typename KeyType, typename ValueType, typename Hash = std::hash<KeyType>, typename Equal = std::equal_to<KeyType>>
+template <typename KeyType, typename ValueType, typename Hash = std::hash<KeyType>, typename Equal = std::equal_to<KeyType>>
 class THashMap
 {
 	using FInnerMap = std::unordered_map<KeyType, ValueType, Hash, Equal, StlAllocator<std::pair<const KeyType, ValueType>>>;
 
-	template<bool bConst>
+	template <bool bConst>
 	class TIterator
 	{
 		using FInnerIterator = std::conditional_t<bConst, typename FInnerMap::const_iterator, typename FInnerMap::iterator>;
-		using FKeyAndValue = std::conditional_t<bConst, const TKeyAndValue<const KeyType, ValueType>, TKeyAndValue<const KeyType, ValueType>>;
+		using FKeyAndValue =
+			std::conditional_t<bConst, const TKeyAndValue<const KeyType, ValueType>, TKeyAndValue<const KeyType, ValueType>>;
+
 	public:
 		// Constructor wrapping inner-iterator
 		TIterator(FInnerIterator InnerIterator) : InnerIterator{InnerIterator} {}
@@ -42,19 +44,17 @@ class THashMap
 		// Move constructor
 		TIterator(TIterator&&) = default;
 		// Copy assignment operator
-		TIterator& operator = (const TIterator& Other) = default;
+		TIterator& operator=(const TIterator& Other) = default;
 		// Move assignment operator
-		TIterator& operator = (TIterator&& Other) = default;
+		TIterator& operator=(TIterator&& Other) = default;
 
 		FKeyAndValue& operator*() const
 		{
 			// Super hacky: Reinterpret a std::pair to a TKeyAndValue so
-			// that THashMap's iterator has the interface of a TMap instead of a 
+			// that THashMap's iterator has the interface of a TMap instead of a
 			// std::unordered_map.
 			// TODO(SOL-7652): Remove this filth when THashMap is reimplemented.
-			using FPair = std::conditional_t<bConst,
-				const std::pair<const KeyType, ValueType>,
-				std::pair<const KeyType, ValueType>>;
+			using FPair = std::conditional_t<bConst, const std::pair<const KeyType, ValueType>, std::pair<const KeyType, ValueType>>;
 			FPair* Pair = &*InnerIterator;
 			FKeyAndValue* KeyAndValue = reinterpret_cast<FKeyAndValue*>(Pair);
 			static_assert(std::is_same_v<decltype(Pair->first), decltype(KeyAndValue->Key)>);
@@ -72,12 +72,12 @@ class THashMap
 			return *this;
 		}
 
-		bool operator == (const TIterator& Other) const
+		bool operator==(const TIterator& Other) const
 		{
 			return InnerIterator == Other.InnerIterator;
 		}
 
-		bool operator != (const TIterator& Other) const
+		bool operator!=(const TIterator& Other) const
 		{
 			return InnerIterator != Other.InnerIterator;
 		}
@@ -103,25 +103,25 @@ public:
 	{
 		if (this != &Other)
 		{
-			Other.Map.clear(); // Ensure Other is cleared.
+			Other.Map.clear();  // Ensure Other is cleared.
 		}
 	}
 	// Copy assignment operator
-	THashMap& operator = (const THashMap& Other) = default;
+	THashMap& operator=(const THashMap& Other) = default;
 	// Move assignment operator
-	THashMap& operator = (THashMap&& Other)
+	THashMap& operator=(THashMap&& Other)
 	{
 		if (this != &Other)
 		{
 			Map = std::move(Other.Map);
-			Other.Map.clear(); // Ensure Other is cleared.
+			Other.Map.clear();  // Ensure Other is cleared.
 		}
 		return *this;
 	}
 
 	// Set the value associated with a key, replacing any existing entry with
 	// the given key.
-	template<typename K, typename V>
+	template <typename K, typename V>
 	void Add(K&& Key, V&& Value)
 	{
 		Map[std::forward<K>(Key)] = std::forward<V>(Value);
@@ -129,9 +129,9 @@ public:
 
 	// Looks up the value with the given key, returning a pointer to the value
 	// if found or nullptr if not found.
-	// Warning: The returned pointer will become invalid if the hash map is 
+	// Warning: The returned pointer will become invalid if the hash map is
 	// modified.
-	template<typename K>
+	template <typename K>
 	ValueType* Find(K&& Key)
 	{
 		typename FInnerMap::iterator It = Map.find(std::forward<K>(Key));
@@ -141,9 +141,9 @@ public:
 	// Looks up the value with the given key, returning a reference to the
 	// existing value if found or a reference to a newly added, zero-initialized
 	// value if not found.
-	// Warning: The returned reference will become invalid if the hash map is 
+	// Warning: The returned reference will become invalid if the hash map is
 	// modified.
-	template<typename K>
+	template <typename K>
 	ValueType& FindOrAdd(K&& Key)
 	{
 		return Map[std::forward<K>(Key)];
@@ -188,15 +188,27 @@ public:
 	}
 
 	// Iterator methods
-	ConstIterator begin() const { return ConstIterator{Map.begin()}; }
-	ConstIterator end() const { return ConstIterator{Map.end()}; }
-	Iterator begin() { return Iterator{Map.begin()}; }
-	Iterator end() { return Iterator{Map.end()}; }
+	ConstIterator begin() const
+	{
+		return ConstIterator{Map.begin()};
+	}
+	ConstIterator end() const
+	{
+		return ConstIterator{Map.end()};
+	}
+	Iterator begin()
+	{
+		return Iterator{Map.begin()};
+	}
+	Iterator end()
+	{
+		return Iterator{Map.end()};
+	}
 
 private:
 	FInnerMap Map;
 };
 
-}  // AutoRTFM
+}  // namespace AutoRTFM
 
-#endif // (defined(__AUTORTFM) && __AUTORTFM)
+#endif  // (defined(__AUTORTFM) && __AUTORTFM)
diff --git a/Engine/Source/Runtime/AutoRTFM/Private/HitSet.h b/Engine/Source/Runtime/AutoRTFM/Private/HitSet.h
--- a/Engine/Source/Runtime/AutoRTFM/Private/HitSet.h
+++ b/Engine/Source/Runtime/AutoRTFM/Private/HitSet.h
@@ -11,79 +11,80 @@
 
 namespace AutoRTFM
 {
-	struct FHitSetEntry
+struct FHitSetEntry
+{
+	static constexpr size_t AddressBits = 48;
+	static constexpr size_t SizeBits = 12;
+	static constexpr size_t FlagsBits = 4;
+
+	// The address of the write.
+	uintptr_t Address : AddressBits;
+	// Size of the write.
+	uintptr_t Size : SizeBits;
+	// Flags applied to the write.
+	uintptr_t Flags : FlagsBits;
+
+	UE_AUTORTFM_FORCEINLINE uintptr_t Payload() const
 	{
-		static constexpr size_t AddressBits = 48;
-		static constexpr size_t SizeBits = 12;
-		static constexpr size_t FlagsBits = 4;
-
-		// The address of the write.
-		uintptr_t Address : AddressBits;
-		// Size of the write.
-		uintptr_t Size : SizeBits;
-		// Flags applied to the write.
-		uintptr_t Flags : FlagsBits;
-
-		UE_AUTORTFM_FORCEINLINE uintptr_t Payload() const
-		{
-			// Note: Clang optimizes this away to a reinterpret cast.
-			uintptr_t Result;
-			memcpy(&Result, this, sizeof(Result));
-			return Result;
-		}
-	};
+		// Note: Clang optimizes this away to a reinterpret cast.
+		uintptr_t Result;
+		memcpy(&Result, this, sizeof(Result));
+		return Result;
+	}
+};
 
-	static_assert(sizeof(FHitSetEntry) == sizeof(uintptr_t));
-	
-	class FHitSet final
-	{
-	public:
-		// Maximum size in bytes for a hit-set write record.
-		// The cutoff here is arbitrarily any number less than UINT16_MAX, but its a
-		// weigh up what a good size is. Because the hitset doesn't detect when you
-		// are trying to write to a subregion of a previous hit (like memset something,
-		// then write to an individual element), we've got to balance the cost of
-		// recording meaningless hits, against the potential to hit again.
-		static constexpr size_t MaxSize = 16;
-
-		// TODO: Revisit a good max capacity for the hitset.
-		static constexpr uint64_t MaxCapacity = static_cast<uint64_t>(128 * 1024 * 1024) / sizeof(FHitSetEntry);
-	private:
-		static constexpr uint8_t LinearProbeDepth = 32;
-		static constexpr uint8_t LogInitialCapacity = 6;
-		static constexpr uint64_t InitialCapacity = static_cast<uint64_t>(1) << LogInitialCapacity;
+static_assert(sizeof(FHitSetEntry) == sizeof(uintptr_t));
 
-		static_assert(LinearProbeDepth <= InitialCapacity);
-	public:
+class FHitSet final
+{
+public:
+	// Maximum size in bytes for a hit-set write record.
+	// The cutoff here is arbitrarily any number less than UINT16_MAX, but its a
+	// weigh up what a good size is. Because the hitset doesn't detect when you
+	// are trying to write to a subregion of a previous hit (like memset something,
+	// then write to an individual element), we've got to balance the cost of
+	// recording meaningless hits, against the potential to hit again.
+	static constexpr size_t MaxSize = 16;
+
+	// TODO: Revisit a good max capacity for the hitset.
+	static constexpr uint64_t MaxCapacity = static_cast<uint64_t>(128 * 1024 * 1024) / sizeof(FHitSetEntry);
+
+private:
+	static constexpr uint8_t LinearProbeDepth = 32;
+	static constexpr uint8_t LogInitialCapacity = 6;
+	static constexpr uint64_t InitialCapacity = static_cast<uint64_t>(1) << LogInitialCapacity;
+
+	static_assert(LinearProbeDepth <= InitialCapacity);
+
+public:
+	~FHitSet()
+	{
+		Reset();
+	}
 
-		~FHitSet()
-		{
-			Reset();
-		}
+	enum class EInsertResult : uint8_t
+	{
+		Exists,
+		Inserted,
+		NotInserted
+	};
 
-		enum class EInsertResult : uint8_t
-		{
-			Exists,
-			Inserted,
-			NotInserted
-		};
-
-		// Insert something in the HitSet, returning whether the key already
-		// exists, it was inserted, or not.
-		UE_AUTORTFM_FORCEINLINE EInsertResult FindOrTryInsertNoResize(FHitSetEntry Entry)
-		{
-			return InsertNoResize(Entry.Payload());
-		}
+	// Insert something in the HitSet, returning whether the key already
+	// exists, it was inserted, or not.
+	UE_AUTORTFM_FORCEINLINE EInsertResult FindOrTryInsertNoResize(FHitSetEntry Entry)
+	{
+		return InsertNoResize(Entry.Payload());
+	}
 
-		// Insert something in the HitSet, returning true if we found the key.
-		// If false is returned, we could have inserted or not - depending on
-		// whether the HitSet has reached its max capacity.
-		UE_AUTORTFM_FORCEINLINE bool FindOrTryInsert(FHitSetEntry Entry)
+	// Insert something in the HitSet, returning true if we found the key.
+	// If false is returned, we could have inserted or not - depending on
+	// whether the HitSet has reached its max capacity.
+	UE_AUTORTFM_FORCEINLINE bool FindOrTryInsert(FHitSetEntry Entry)
+	{
+		while (true)
 		{
-			while (true)
+			switch (InsertNoResize(Entry.Payload()))
 			{
-				switch (InsertNoResize(Entry.Payload()))
-				{
 				default:
 					InternalUnreachable();
 				case EInsertResult::Inserted:
@@ -96,305 +97,310 @@ namespace AutoRTFM
 						return false;
 					}
 					continue;
-				}
 			}
 		}
+	}
 
-		// Returns true if the HitSet contains the given entry, without mutation.
-		UE_AUTORTFM_FORCEINLINE bool Contains(FHitSetEntry Entry) const
-		{
-			uintptr_t Raw = Entry.Payload();
-			const uintptr_t Hash = FibonacciHash(Raw);
+	// Returns true if the HitSet contains the given entry, without mutation.
+	UE_AUTORTFM_FORCEINLINE bool Contains(FHitSetEntry Entry) const
+	{
+		uintptr_t Raw = Entry.Payload();
+		const uintptr_t Hash = FibonacciHash(Raw);
 
-			// The first index will always be in the range of capacity because we
-			// are using a fibonacci hash.
+		// The first index will always be in the range of capacity because we
+		// are using a fibonacci hash.
+		{
+			const uintptr_t Value = Payload[Hash];
+			if (0 == Value)
 			{
-				const uintptr_t Value = Payload[Hash];
-				if (0 == Value)
-				{
-					return false;
-				}
-				if (Raw == Value)
-				{
-					return true;
-				}
-				// Otherwise we need to do a linear probe...
+				return false;
+			}
+			if (Raw == Value)
+			{
+				return true;
 			}
+			// Otherwise we need to do a linear probe...
+		}
 
-			const uintptr_t Mask = Capacity() - 1;
+		const uintptr_t Mask = Capacity() - 1;
 
-			// Clang goes way over the top with loop unrolling, and makes this code noticeably
-			// slower as a result. So we disable it!
+		// Clang goes way over the top with loop unrolling, and makes this code noticeably
+		// slower as a result. So we disable it!
 #ifdef __clang__
 #pragma clang loop unroll(disable)
 #endif
-			for (uint8_t D = 1; D < LinearProbeDepth; D++)
-			{
-				// Capacity is always a power of 2, so we can just mask out the bits.
-				const uintptr_t I = (Hash + D) & Mask;
-
-				const uintptr_t Value = Payload[I];
-				if (0 == Value)
-				{
-					return false;
-				}
-				if (Raw == Value)
-				{
-					return true;
-				}
+		for (uint8_t D = 1; D < LinearProbeDepth; D++)
+		{
+			// Capacity is always a power of 2, so we can just mask out the bits.
+			const uintptr_t I = (Hash + D) & Mask;
 
-				// Otherwise we need to keep going with our linear probe...
+			const uintptr_t Value = Payload[I];
+			if (0 == Value)
+			{
+				return false;
+			}
+			if (Raw == Value)
+			{
+				return true;
 			}
 
-			return false;
+			// Otherwise we need to keep going with our linear probe...
 		}
 
-		UE_AUTORTFM_FORCEINLINE bool IsEmpty() const
-		{
-			return 0 == Count;
-		}
+		return false;
+	}
 
-		// Clear out the data stored in the set, but does not reduce the capacity.
-		void Reset()
+	UE_AUTORTFM_FORCEINLINE bool IsEmpty() const
+	{
+		return 0 == Count;
+	}
+
+	// Clear out the data stored in the set, but does not reduce the capacity.
+	void Reset()
+	{
+		if (Count)
 		{
-			if (Count)
+			if (Payload != SmallPayload)
 			{
-				if (Payload != SmallPayload)
-				{
-					AutoRTFM::Free(Payload);
-					Payload = SmallPayload;
-				}
-				memset(SmallPayload, 0, sizeof(SmallPayload));
-				SixtyFourMinusLogCapacity = 64 - LogInitialCapacity;
-				Count = 0;
+				AutoRTFM::Free(Payload);
+				Payload = SmallPayload;
 			}
+			memset(SmallPayload, 0, sizeof(SmallPayload));
+			SixtyFourMinusLogCapacity = 64 - LogInitialCapacity;
+			Count = 0;
 		}
+	}
 
-		uint64_t GetCapacity() const
+	uint64_t GetCapacity() const
+	{
+		return Capacity();
+	}
+
+	uint64_t GetCount() const
+	{
+		return Count;
+	}
+
+	struct FIterator
+	{
+		FIterator(uintptr_t const* Ptr, size_t EntriesRemaining) : Ptr{Ptr}, EntriesRemaining{EntriesRemaining}
 		{
-			return Capacity();
+			SkipEmpty();
 		}
 
-		uint64_t GetCount() const
+		FHitSetEntry operator*() const
 		{
-			return Count;
+			FHitSetEntry Out;
+			memcpy(&Out, Ptr, sizeof(Out));
+			return Out;
 		}
 
-		struct FIterator
+		FIterator& operator++()
 		{
-			FIterator(uintptr_t const * Ptr, size_t EntriesRemaining)
-				: Ptr{Ptr}, EntriesRemaining{EntriesRemaining}
+			if (EntriesRemaining)
 			{
+				EntriesRemaining--;
+				Ptr++;
 				SkipEmpty();
 			}
+			return *this;
+		}
 
-			FHitSetEntry operator*() const
-			{
-				FHitSetEntry Out;
-				memcpy(&Out, Ptr, sizeof(Out));
-				return Out;
-			}
+		bool operator==(FIterator Other) const
+		{
+			return Other.Ptr == Ptr;
+		}
 
-			FIterator& operator++()
-			{
-				if (EntriesRemaining)
-				{
-					EntriesRemaining--;
-					Ptr++;
-					SkipEmpty();
-				}
-				return *this;
-			}
+		bool operator!=(FIterator Other) const
+		{
+			return Other.Ptr != Ptr;
+		}
 
-			bool operator==(FIterator Other) const
+	private:
+		void SkipEmpty()
+		{
+			while (EntriesRemaining > 0 && *Ptr == 0)
 			{
-				return Other.Ptr == Ptr;
+				EntriesRemaining--;
+				Ptr++;
 			}
+		}
 
-			bool operator!=(FIterator Other) const
-			{
-				return Other.Ptr != Ptr;
-			}
+		const uintptr_t* Ptr;
+		size_t EntriesRemaining;
+	};
 
-		private:
-			void SkipEmpty()
-			{
-				while (EntriesRemaining > 0 && *Ptr == 0)
-				{
-					EntriesRemaining--;
-					Ptr++;
-				}
-			}
+	FIterator begin() const
+	{
+		return FIterator(Payload, Capacity());
+	}
+	FIterator end() const
+	{
+		return FIterator(Payload + Capacity(), 0);
+	}
 
-			const uintptr_t* Ptr;
-			size_t EntriesRemaining;
-		};
-		
-		FIterator begin() const { return FIterator(Payload, Capacity()); }
-		FIterator end() const { return FIterator(Payload + Capacity(), 0); }
+private:
+	uintptr_t SmallPayload[InitialCapacity]{};
+	uintptr_t* Payload = SmallPayload;
+	uint64_t Count = 0;
+	uint32_t SixtyFourMinusLogCapacity = 64 - LogInitialCapacity;
 
-	private:
-		uintptr_t SmallPayload[InitialCapacity]{};
-		uintptr_t* Payload = SmallPayload;
-		uint64_t Count = 0;
-		uint32_t SixtyFourMinusLogCapacity = 64 - LogInitialCapacity;
+	UE_AUTORTFM_FORCEINLINE uint64_t Capacity() const
+	{
+		UE_AUTORTFM_ASSUME(SixtyFourMinusLogCapacity < 64);
+		return static_cast<uint64_t>(1) << (64 - SixtyFourMinusLogCapacity);
+	}
 
-		UE_AUTORTFM_FORCEINLINE uint64_t Capacity() const
-		{
-			UE_AUTORTFM_ASSUME(SixtyFourMinusLogCapacity < 64);
-			return static_cast<uint64_t>(1) << (64 - SixtyFourMinusLogCapacity);
-		}
+	UE_AUTORTFM_FORCEINLINE void IncreaseCapacity()
+	{
+		// It seems odd - but subtracting 1 here is totally intentional
+		// because of how we store our capacity as (64 - log2(Capacity)).
+		SixtyFourMinusLogCapacity -= 1;
 
-		UE_AUTORTFM_FORCEINLINE void IncreaseCapacity()
-		{
-			// It seems odd - but subtracting 1 here is totally intentional
-			// because of how we store our capacity as (64 - log2(Capacity)).
-			SixtyFourMinusLogCapacity -= 1;
+		// Check that we haven't overflowed our capacity!
+		AUTORTFM_ASSERT(0 != SixtyFourMinusLogCapacity);
+	}
 
-			// Check that we haven't overflowed our capacity!
-			AUTORTFM_ASSERT(0 != SixtyFourMinusLogCapacity);
-		}
+	bool Resize()
+	{
+		const uint64_t OldCapacity = Capacity();
 
-		bool Resize()
+		if (OldCapacity >= MaxCapacity)
 		{
-			const uint64_t OldCapacity = Capacity();
-
-			if (OldCapacity >= MaxCapacity)
-			{
-				return false;
-			}
+			return false;
+		}
 
-			uintptr_t* const OldPayload = Payload;
-			const uint64_t OldCount = Count;
+		uintptr_t* const OldPayload = Payload;
+		const uint64_t OldCount = Count;
 
-			while (true)
-			{
-				bool bNeedAnotherResize = false;
+		while (true)
+		{
+			bool bNeedAnotherResize = false;
 
-				Count = 0;
+			Count = 0;
 
-				IncreaseCapacity();
+			IncreaseCapacity();
 
-				Payload = static_cast<uintptr_t*>(AutoRTFM::AllocateZeroed(Capacity() * sizeof(uintptr_t), alignof(uintptr_t)));
-				AUTORTFM_ASSERT(nullptr != Payload);
+			Payload = static_cast<uintptr_t*>(AutoRTFM::AllocateZeroed(Capacity() * sizeof(uintptr_t), alignof(uintptr_t)));
+			AUTORTFM_ASSERT(nullptr != Payload);
 
-				// Now we need to rehash and reinsert all the items.
-				for (size_t I = 0; I < OldCapacity; I++)
+			// Now we need to rehash and reinsert all the items.
+			for (size_t I = 0; I < OldCapacity; I++)
+			{
+				// Skip empty locations.
+				if (0 == OldPayload[I])
 				{
-					// Skip empty locations.
-					if (0 == OldPayload[I])
-					{
-						continue;
-					}
+					continue;
+				}
 
-					const EInsertResult Result = InsertNoResize(OldPayload[I]);
+				const EInsertResult Result = InsertNoResize(OldPayload[I]);
 
-					if (EInsertResult::NotInserted == Result)
-					{
-						bNeedAnotherResize = true;
-						break;
-					}
-					else
-					{
-						AUTORTFM_ASSERT(EInsertResult::Inserted == Result);
-					}
+				if (EInsertResult::NotInserted == Result)
+				{
+					bNeedAnotherResize = true;
+					break;
 				}
-
-				if (bNeedAnotherResize)
+				else
 				{
-					if (Payload != SmallPayload)
-					{
-						AutoRTFM::Free(Payload);
-					}
-					continue;
+					AUTORTFM_ASSERT(EInsertResult::Inserted == Result);
 				}
-
-				break;
 			}
 
-			AUTORTFM_ASSERT(OldCount == Count);
-
-			if (OldPayload != SmallPayload)
+			if (bNeedAnotherResize)
 			{
-				AutoRTFM::Free(OldPayload);
+				if (Payload != SmallPayload)
+				{
+					AutoRTFM::Free(Payload);
+				}
+				continue;
 			}
 
-			return true;
+			break;
 		}
 
-		UE_AUTORTFM_FORCEINLINE uintptr_t FibonacciHash(const uintptr_t Hashee) const
+		AUTORTFM_ASSERT(OldCount == Count);
+
+		if (OldPayload != SmallPayload)
 		{
-			constexpr uintptr_t Fibonacci = 0x9E3779B97F4A7C15;
-			return (Hashee * Fibonacci) >> SixtyFourMinusLogCapacity;
+			AutoRTFM::Free(OldPayload);
 		}
 
-		UE_AUTORTFM_FORCEINLINE EInsertResult TryInsertAtIndex(const uintptr_t Raw, const uintptr_t I)
-		{
-			// We have a free location in the set.
-			if (0 == Payload[I])
-			{
-				Payload[I] = Raw;
-				Count++;
-				return EInsertResult::Inserted;
-			}
+		return true;
+	}
 
-			// We're already in the set.
-			if (Raw == Payload[I])
-			{
-				return EInsertResult::Exists;
-			}
+	UE_AUTORTFM_FORCEINLINE uintptr_t FibonacciHash(const uintptr_t Hashee) const
+	{
+		constexpr uintptr_t Fibonacci = 0x9E3779B97F4A7C15;
+		return (Hashee * Fibonacci) >> SixtyFourMinusLogCapacity;
+	}
 
-			return EInsertResult::NotInserted;
+	UE_AUTORTFM_FORCEINLINE EInsertResult TryInsertAtIndex(const uintptr_t Raw, const uintptr_t I)
+	{
+		// We have a free location in the set.
+		if (0 == Payload[I])
+		{
+			Payload[I] = Raw;
+			Count++;
+			return EInsertResult::Inserted;
 		}
 
-		// Insert something in the HitSet, returning true if the insert
-		// succeeded (EG. the key was not already in the set).
-		UE_AUTORTFM_FORCEINLINE EInsertResult InsertNoResize(const uintptr_t Raw)
+		// We're already in the set.
+		if (Raw == Payload[I])
 		{
-			AUTORTFM_ASSERT_DEBUG(Raw != 0);
+			return EInsertResult::Exists;
+		}
 
-			const uintptr_t Hash = FibonacciHash(Raw);
+		return EInsertResult::NotInserted;
+	}
 
-			// The first index will always be in the range of capacity because we
-			// are using a fibonacci hash. So just check if we can insert here.
-			{
-				const uintptr_t I = Hash;
+	// Insert something in the HitSet, returning true if the insert
+	// succeeded (EG. the key was not already in the set).
+	UE_AUTORTFM_FORCEINLINE EInsertResult InsertNoResize(const uintptr_t Raw)
+	{
+		AUTORTFM_ASSERT_DEBUG(Raw != 0);
 
-				const EInsertResult R = TryInsertAtIndex(Raw, I);
+		const uintptr_t Hash = FibonacciHash(Raw);
 
-				if (EInsertResult::NotInserted != R)
-				{
-					return R;
-				}
+		// The first index will always be in the range of capacity because we
+		// are using a fibonacci hash. So just check if we can insert here.
+		{
+			const uintptr_t I = Hash;
+
+			const EInsertResult R = TryInsertAtIndex(Raw, I);
 
-				// Otherwise we need to do a linear probe...
+			if (EInsertResult::NotInserted != R)
+			{
+				return R;
 			}
 
-			const uintptr_t Mask = Capacity() - 1;
+			// Otherwise we need to do a linear probe...
+		}
+
+		const uintptr_t Mask = Capacity() - 1;
 
-			// Clang goes way over the top with loop unrolling, and makes this code noticeably
-			// slower as a result. So we disable it!
+		// Clang goes way over the top with loop unrolling, and makes this code noticeably
+		// slower as a result. So we disable it!
 #ifdef __clang__
 #pragma clang loop unroll(disable)
 #endif
-			for (uint8_t D = 1; D < LinearProbeDepth; D++)
-			{
-				// Capacity is always a power of 2, so we can just mask out the bits.
-				const uintptr_t I = (Hash + D) & Mask;
+		for (uint8_t D = 1; D < LinearProbeDepth; D++)
+		{
+			// Capacity is always a power of 2, so we can just mask out the bits.
+			const uintptr_t I = (Hash + D) & Mask;
 
-				const EInsertResult R = TryInsertAtIndex(Raw, I);
+			const EInsertResult R = TryInsertAtIndex(Raw, I);
 
-				if (EInsertResult::NotInserted != R)
-				{
-					return R;
-				}
-
-				// Otherwise we need to keep going with our linear probe...
+			if (EInsertResult::NotInserted != R)
+			{
+				return R;
 			}
 
-			return EInsertResult::NotInserted;
+			// Otherwise we need to keep going with our linear probe...
 		}
-	};
-}
 
-#endif // (defined(__AUTORTFM) && __AUTORTFM)
+		return EInsertResult::NotInserted;
+	}
+};
+}  // namespace AutoRTFM
+
+#endif  // (defined(__AUTORTFM) && __AUTORTFM)
diff --git a/Engine/Source/Runtime/AutoRTFM/Private/Stack.h b/Engine/Source/Runtime/AutoRTFM/Private/Stack.h
--- a/Engine/Source/Runtime/AutoRTFM/Private/Stack.h
+++ b/Engine/Source/Runtime/AutoRTFM/Private/Stack.h
@@ -30,7 +30,7 @@ namespace AutoRTFM
 //   InlineCapacity - the number of elements that can be held before spilling to
 //                    the heap.
 //   Validation - if Disabled, do not perform validity assertions.
-template<typename T, size_t InlineCapacity, EContainerValidation Validation = EContainerValidation::Enabled>
+template <typename T, size_t InlineCapacity, EContainerValidation Validation = EContainerValidation::Enabled>
 class AUTORTFM_DISABLE_NO_SANITIZE TStack
 {
 public:
@@ -89,7 +89,7 @@ public:
 		}
 		return *this;
 	}
-	
+
 	// Clears all the items from the stack, preserving the capacity.
 	void Clear()
 	{
@@ -164,7 +164,7 @@ public:
 	{
 		if (NewCapacity <= Capacity)
 		{
-			return; // Already has space for the new capacity
+			return;  // Already has space for the new capacity
 		}
 
 		if constexpr (std::is_trivial_v<T>)
@@ -198,9 +198,15 @@ public:
 		Capacity = NewCapacity;
 	}
 
-	inline size_t Num() const { return Count; }
+	inline size_t Num() const
+	{
+		return Count;
+	}
 
-	inline bool IsEmpty() const { return Count == 0; }
+	inline bool IsEmpty() const
+	{
+		return Count == 0;
+	}
 
 	T& operator[](size_t Index)
 	{
@@ -238,14 +244,26 @@ public:
 		return Data[Count - 1];
 	}
 
-	T* begin() { return Data; }
-	const T* begin() const { return Data; }
-	T* end() { return Data + Count; }
-	const T* end() const { return Data + Count; }
+	T* begin()
+	{
+		return Data;
+	}
+	const T* begin() const
+	{
+		return Data;
+	}
+	T* end()
+	{
+		return Data + Count;
+	}
+	const T* end() const
+	{
+		return Data + Count;
+	}
 
 	// Equality operator
-	template<typename OtherT, size_t OtherInlineCapacity>
-	bool operator == (const TStack<OtherT, OtherInlineCapacity>& Other) const
+	template <typename OtherT, size_t OtherInlineCapacity>
+	bool operator==(const TStack<OtherT, OtherInlineCapacity>& Other) const
 	{
 		if (Num() != Other.Num())
 		{
@@ -321,8 +339,14 @@ private:
 		Count = 0;
 	}
 
-	T* InlineData() { return reinterpret_cast<T*>(InlineDataBuffer); }
-	const T* InlineData() const { return reinterpret_cast<const T*>(InlineDataBuffer); }
+	T* InlineData()
+	{
+		return reinterpret_cast<T*>(InlineDataBuffer);
+	}
+	const T* InlineData() const
+	{
+		return reinterpret_cast<const T*>(InlineDataBuffer);
+	}
 
 	alignas(T) std::byte InlineDataBuffer[sizeof(T) * InlineCapacity];
 	T* Data = reinterpret_cast<T*>(InlineDataBuffer);
@@ -330,6 +354,6 @@ private:
 	size_t Count = 0;
 };
 
-}
+}  // namespace AutoRTFM
 
-#endif // (defined(__AUTORTFM) && __AUTORTFM)
+#endif  // (defined(__AUTORTFM) && __AUTORTFM)
diff --git a/Engine/Source/Runtime/AutoRTFM/Private/Stats.cpp b/Engine/Source/Runtime/AutoRTFM/Private/Stats.cpp
--- a/Engine/Source/Runtime/AutoRTFM/Private/Stats.cpp
+++ b/Engine/Source/Runtime/AutoRTFM/Private/Stats.cpp
@@ -8,114 +8,134 @@
 
 #include <inttypes.h>
 
-template<> void AutoRTFM::FStats::Report<AutoRTFM::EStatsKind::Transaction>(const uint64_t Data) const
+template <>
+void AutoRTFM::FStats::Report<AutoRTFM::EStatsKind::Transaction>(const uint64_t Data) const
 {
 	AUTORTFM_LOG("  Total transactions:        %" PRIu64, Data);
 }
 
-template<> void AutoRTFM::FStats::Report<AutoRTFM::EStatsKind::Commit>(const uint64_t Data) const
+template <>
+void AutoRTFM::FStats::Report<AutoRTFM::EStatsKind::Commit>(const uint64_t Data) const
 {
 	AUTORTFM_LOG("  Total commits:             %" PRIu64, Data);
 }
 
-template<> void AutoRTFM::FStats::Report<AutoRTFM::EStatsKind::Abort>(const uint64_t Data) const
+template <>
+void AutoRTFM::FStats::Report<AutoRTFM::EStatsKind::Abort>(const uint64_t Data) const
 {
 	AUTORTFM_LOG("  Total aborts:              %" PRIu64, Data);
 }
 
-template<> void AutoRTFM::FStats::Report<AutoRTFM::EStatsKind::AverageTransactionDepth>(const uint64_t Data) const
+template <>
+void AutoRTFM::FStats::Report<AutoRTFM::EStatsKind::AverageTransactionDepth>(const uint64_t Data) const
 {
 	const uint64_t TotalTransactions = Datas[static_cast<size_t>(AutoRTFM::EStatsKind::Transaction)];
 	AUTORTFM_LOG("  Average transaction depth: %15.3f", (static_cast<double>(Data) / static_cast<double>(TotalTransactions)));
 }
 
-template<> void AutoRTFM::FStats::Report<AutoRTFM::EStatsKind::MaximumTransactionDepth>(const uint64_t Data) const
+template <>
+void AutoRTFM::FStats::Report<AutoRTFM::EStatsKind::MaximumTransactionDepth>(const uint64_t Data) const
 {
-    AUTORTFM_LOG("  Maximum transaction depth: %" PRIu64, Data);
+	AUTORTFM_LOG("  Maximum transaction depth: %" PRIu64, Data);
 }
 
-template<> void AutoRTFM::FStats::Report<AutoRTFM::EStatsKind::AverageWriteLogEntries>(const uint64_t Data) const
+template <>
+void AutoRTFM::FStats::Report<AutoRTFM::EStatsKind::AverageWriteLogEntries>(const uint64_t Data) const
 {
-    const uint64_t TotalTransactions = Datas[static_cast<size_t>(AutoRTFM::EStatsKind::Transaction)];
-    AUTORTFM_LOG("  Average write log entries: %15.3f", (static_cast<double>(Data) / static_cast<double>(TotalTransactions)));
+	const uint64_t TotalTransactions = Datas[static_cast<size_t>(AutoRTFM::EStatsKind::Transaction)];
+	AUTORTFM_LOG("  Average write log entries: %15.3f", (static_cast<double>(Data) / static_cast<double>(TotalTransactions)));
 }
 
-template<> void AutoRTFM::FStats::Report<AutoRTFM::EStatsKind::MaximumWriteLogEntries>(const uint64_t Data) const
+template <>
+void AutoRTFM::FStats::Report<AutoRTFM::EStatsKind::MaximumWriteLogEntries>(const uint64_t Data) const
 {
-    AUTORTFM_LOG("  Maximum write log entries: %" PRIu64, Data);
+	AUTORTFM_LOG("  Maximum write log entries: %" PRIu64, Data);
 }
 
-template<> void AutoRTFM::FStats::Report<AutoRTFM::EStatsKind::AverageWriteLogBytes>(const uint64_t Data) const
+template <>
+void AutoRTFM::FStats::Report<AutoRTFM::EStatsKind::AverageWriteLogBytes>(const uint64_t Data) const
 {
-    const uint64_t TotalTransactions = Datas[static_cast<size_t>(AutoRTFM::EStatsKind::Transaction)];
-    AUTORTFM_LOG("  Average write log bytes:   %15.3f", (static_cast<double>(Data) / static_cast<double>(TotalTransactions)));
+	const uint64_t TotalTransactions = Datas[static_cast<size_t>(AutoRTFM::EStatsKind::Transaction)];
+	AUTORTFM_LOG("  Average write log bytes:   %15.3f", (static_cast<double>(Data) / static_cast<double>(TotalTransactions)));
 }
 
-template<> void AutoRTFM::FStats::Report<AutoRTFM::EStatsKind::MaximumWriteLogBytes>(const uint64_t Data) const
+template <>
+void AutoRTFM::FStats::Report<AutoRTFM::EStatsKind::MaximumWriteLogBytes>(const uint64_t Data) const
 {
-    AUTORTFM_LOG("  Maximum write log bytes:   %" PRIu64, Data);
+	AUTORTFM_LOG("  Maximum write log bytes:   %" PRIu64, Data);
 }
 
-template<> void AutoRTFM::FStats::Report<AutoRTFM::EStatsKind::HitSetHit>(const uint64_t Data) const
+template <>
+void AutoRTFM::FStats::Report<AutoRTFM::EStatsKind::HitSetHit>(const uint64_t Data) const
 {
-    AUTORTFM_LOG("  HitSet hits:               %" PRIu64, Data);
+	AUTORTFM_LOG("  HitSet hits:               %" PRIu64, Data);
 }
 
-template<> void AutoRTFM::FStats::Report<AutoRTFM::EStatsKind::HitSetMiss>(const uint64_t Data) const
+template <>
+void AutoRTFM::FStats::Report<AutoRTFM::EStatsKind::HitSetMiss>(const uint64_t Data) const
 {
-    AUTORTFM_LOG("  HitSet misses:             %" PRIu64, Data);
+	AUTORTFM_LOG("  HitSet misses:             %" PRIu64, Data);
 }
 
-template<> void AutoRTFM::FStats::Report<AutoRTFM::EStatsKind::HitSetSkippedBecauseOfStackLocalMemory>(const uint64_t Data) const
+template <>
+void AutoRTFM::FStats::Report<AutoRTFM::EStatsKind::HitSetSkippedBecauseOfStackLocalMemory>(const uint64_t Data) const
 {
-    AUTORTFM_LOG("  HitSet skip (stack local): %" PRIu64, Data);
+	AUTORTFM_LOG("  HitSet skip (stack local): %" PRIu64, Data);
 }
 
-template<> void AutoRTFM::FStats::Report<AutoRTFM::EStatsKind::AverageCommitTasks>(const uint64_t Data) const
+template <>
+void AutoRTFM::FStats::Report<AutoRTFM::EStatsKind::AverageCommitTasks>(const uint64_t Data) const
 {
-    const uint64_t TotalTransactions = Datas[static_cast<size_t>(AutoRTFM::EStatsKind::Transaction)];
-    AUTORTFM_LOG("  Average commit tasks:      %15.3f", (static_cast<double>(Data) / static_cast<double>(TotalTransactions)));
+	const uint64_t TotalTransactions = Datas[static_cast<size_t>(AutoRTFM::EStatsKind::Transaction)];
+	AUTORTFM_LOG("  Average commit tasks:      %15.3f", (static_cast<double>(Data) / static_cast<double>(TotalTransactions)));
 }
 
-template<> void AutoRTFM::FStats::Report<AutoRTFM::EStatsKind::MaximumCommitTasks>(const uint64_t Data) const
+template <>
+void AutoRTFM::FStats::Report<AutoRTFM::EStatsKind::MaximumCommitTasks>(const uint64_t Data) const
 {
-    AUTORTFM_LOG("  Maximum commit tasks:      %" PRIu64, Data);
+	AUTORTFM_LOG("  Maximum commit tasks:      %" PRIu64, Data);
 }
 
-template<> void AutoRTFM::FStats::Report<AutoRTFM::EStatsKind::AverageAbortTasks>(const uint64_t Data) const
+template <>
+void AutoRTFM::FStats::Report<AutoRTFM::EStatsKind::AverageAbortTasks>(const uint64_t Data) const
 {
-    const uint64_t TotalTransactions = Datas[static_cast<size_t>(AutoRTFM::EStatsKind::Transaction)];
-    AUTORTFM_LOG("  Average abort tasks:       %15.3f", (static_cast<double>(Data) / static_cast<double>(TotalTransactions)));
+	const uint64_t TotalTransactions = Datas[static_cast<size_t>(AutoRTFM::EStatsKind::Transaction)];
+	AUTORTFM_LOG("  Average abort tasks:       %15.3f", (static_cast<double>(Data) / static_cast<double>(TotalTransactions)));
 }
 
-template<> void AutoRTFM::FStats::Report<AutoRTFM::EStatsKind::MaximumAbortTasks>(const uint64_t Data) const
+template <>
+void AutoRTFM::FStats::Report<AutoRTFM::EStatsKind::MaximumAbortTasks>(const uint64_t Data) const
 {
-    AUTORTFM_LOG("  Maximum abort tasks:       %" PRIu64, Data);
+	AUTORTFM_LOG("  Maximum abort tasks:       %" PRIu64, Data);
 }
 
-template<> void AutoRTFM::FStats::Report<AutoRTFM::EStatsKind::NewMemoryTrackerHit>(const uint64_t Data) const
+template <>
+void AutoRTFM::FStats::Report<AutoRTFM::EStatsKind::NewMemoryTrackerHit>(const uint64_t Data) const
 {
-    AUTORTFM_LOG("  New memory hits:           %" PRIu64, Data);
+	AUTORTFM_LOG("  New memory hits:           %" PRIu64, Data);
 }
 
-template<> void AutoRTFM::FStats::Report<AutoRTFM::EStatsKind::NewMemoryTrackerMiss>(const uint64_t Data) const
+template <>
+void AutoRTFM::FStats::Report<AutoRTFM::EStatsKind::NewMemoryTrackerMiss>(const uint64_t Data) const
 {
-    AUTORTFM_LOG("  New memory misses:         %" PRIu64, Data);
+	AUTORTFM_LOG("  New memory misses:         %" PRIu64, Data);
 }
 
-template<> void AutoRTFM::FStats::Report<AutoRTFM::EStatsKind::AverageHitSetSize>(const uint64_t Data) const
+template <>
+void AutoRTFM::FStats::Report<AutoRTFM::EStatsKind::AverageHitSetSize>(const uint64_t Data) const
 {
-    const uint64_t TotalTransactions = Datas[static_cast<size_t>(AutoRTFM::EStatsKind::Transaction)];
-    AUTORTFM_LOG("  Average hit set size:      %15.3f", (static_cast<double>(Data) / static_cast<double>(TotalTransactions)));
+	const uint64_t TotalTransactions = Datas[static_cast<size_t>(AutoRTFM::EStatsKind::Transaction)];
+	AUTORTFM_LOG("  Average hit set size:      %15.3f", (static_cast<double>(Data) / static_cast<double>(TotalTransactions)));
 }
 
-template<> void AutoRTFM::FStats::Report<AutoRTFM::EStatsKind::AverageHitSetCapacity>(const uint64_t Data) const
+template <>
+void AutoRTFM::FStats::Report<AutoRTFM::EStatsKind::AverageHitSetCapacity>(const uint64_t Data) const
 {
-    const uint64_t TotalTransactions = Datas[static_cast<size_t>(AutoRTFM::EStatsKind::Transaction)];
-    AUTORTFM_LOG("  Average hit set capacity:  %15.3f", (static_cast<double>(Data) / static_cast<double>(TotalTransactions)));
+	const uint64_t TotalTransactions = Datas[static_cast<size_t>(AutoRTFM::EStatsKind::Transaction)];
+	AUTORTFM_LOG("  Average hit set capacity:  %15.3f", (static_cast<double>(Data) / static_cast<double>(TotalTransactions)));
 
-    const uint64_t TotalHitSetSize = Datas[static_cast<size_t>(AutoRTFM::EStatsKind::AverageHitSetSize)];
-    AUTORTFM_LOG("  Average hit set occupancy: %15.3f", 100 * (static_cast<double>(TotalHitSetSize) / static_cast<double>(Data)));
+	const uint64_t TotalHitSetSize = Datas[static_cast<size_t>(AutoRTFM::EStatsKind::AverageHitSetSize)];
+	AUTORTFM_LOG("  Average hit set occupancy: %15.3f", 100 * (static_cast<double>(TotalHitSetSize) / static_cast<double>(Data)));
 }
 
 void AutoRTFM::FStats::Report() const
@@ -128,28 +148,33 @@ void AutoRTFM::FStats::Report() const
 		{
 			switch (static_cast<EStatsKind>(I))
 			{
-			default: AutoRTFM::InternalUnreachable(); break;
-#define REPORT_CASE(x) case (x): Report<x>(Datas[static_cast<size_t>(x)]); break
-			REPORT_CASE(AutoRTFM::EStatsKind::Transaction);
-			REPORT_CASE(AutoRTFM::EStatsKind::Commit);
-			REPORT_CASE(AutoRTFM::EStatsKind::Abort);
-			REPORT_CASE(AutoRTFM::EStatsKind::AverageTransactionDepth);
-            REPORT_CASE(AutoRTFM::EStatsKind::MaximumTransactionDepth);
-            REPORT_CASE(AutoRTFM::EStatsKind::AverageWriteLogEntries);
-            REPORT_CASE(AutoRTFM::EStatsKind::MaximumWriteLogEntries);
-			REPORT_CASE(AutoRTFM::EStatsKind::AverageWriteLogBytes);
-			REPORT_CASE(AutoRTFM::EStatsKind::MaximumWriteLogBytes);
-            REPORT_CASE(AutoRTFM::EStatsKind::HitSetHit);
-            REPORT_CASE(AutoRTFM::EStatsKind::HitSetMiss);
-            REPORT_CASE(AutoRTFM::EStatsKind::HitSetSkippedBecauseOfStackLocalMemory);
-            REPORT_CASE(AutoRTFM::EStatsKind::AverageCommitTasks);
-            REPORT_CASE(AutoRTFM::EStatsKind::MaximumCommitTasks);
-            REPORT_CASE(AutoRTFM::EStatsKind::AverageAbortTasks);
-            REPORT_CASE(AutoRTFM::EStatsKind::MaximumAbortTasks);
-            REPORT_CASE(AutoRTFM::EStatsKind::NewMemoryTrackerHit);
-            REPORT_CASE(AutoRTFM::EStatsKind::NewMemoryTrackerMiss);
-			REPORT_CASE(AutoRTFM::EStatsKind::AverageHitSetSize);
-			REPORT_CASE(AutoRTFM::EStatsKind::AverageHitSetCapacity);
+				default:
+					AutoRTFM::InternalUnreachable();
+					break;
+#define REPORT_CASE(x)                            \
+	case (x):                                     \
+		Report<x>(Datas[static_cast<size_t>(x)]); \
+		break
+					REPORT_CASE(AutoRTFM::EStatsKind::Transaction);
+					REPORT_CASE(AutoRTFM::EStatsKind::Commit);
+					REPORT_CASE(AutoRTFM::EStatsKind::Abort);
+					REPORT_CASE(AutoRTFM::EStatsKind::AverageTransactionDepth);
+					REPORT_CASE(AutoRTFM::EStatsKind::MaximumTransactionDepth);
+					REPORT_CASE(AutoRTFM::EStatsKind::AverageWriteLogEntries);
+					REPORT_CASE(AutoRTFM::EStatsKind::MaximumWriteLogEntries);
+					REPORT_CASE(AutoRTFM::EStatsKind::AverageWriteLogBytes);
+					REPORT_CASE(AutoRTFM::EStatsKind::MaximumWriteLogBytes);
+					REPORT_CASE(AutoRTFM::EStatsKind::HitSetHit);
+					REPORT_CASE(AutoRTFM::EStatsKind::HitSetMiss);
+					REPORT_CASE(AutoRTFM::EStatsKind::HitSetSkippedBecauseOfStackLocalMemory);
+					REPORT_CASE(AutoRTFM::EStatsKind::AverageCommitTasks);
+					REPORT_CASE(AutoRTFM::EStatsKind::MaximumCommitTasks);
+					REPORT_CASE(AutoRTFM::EStatsKind::AverageAbortTasks);
+					REPORT_CASE(AutoRTFM::EStatsKind::MaximumAbortTasks);
+					REPORT_CASE(AutoRTFM::EStatsKind::NewMemoryTrackerHit);
+					REPORT_CASE(AutoRTFM::EStatsKind::NewMemoryTrackerMiss);
+					REPORT_CASE(AutoRTFM::EStatsKind::AverageHitSetSize);
+					REPORT_CASE(AutoRTFM::EStatsKind::AverageHitSetCapacity);
 #undef REPORT_CASE
 			}
 		}
@@ -158,4 +183,4 @@ void AutoRTFM::FStats::Report() const
 
 AutoRTFM::FStats AutoRTFM::Stats;
 
-#endif // defined(__AUTORTFM) && __AUTORTFM
+#endif  // defined(__AUTORTFM) && __AUTORTFM
diff --git a/Engine/Source/Runtime/AutoRTFM/Private/StdLibSupport.cpp b/Engine/Source/Runtime/AutoRTFM/Private/StdLibSupport.cpp
--- a/Engine/Source/Runtime/AutoRTFM/Private/StdLibSupport.cpp
+++ b/Engine/Source/Runtime/AutoRTFM/Private/StdLibSupport.cpp
@@ -11,7 +11,7 @@
 #include <algorithm>
 #include <charconv>
 #include <float.h>
-#include <functional> // note: introduces additional math overloads
+#include <functional>  // note: introduces additional math overloads
 #include <stdarg.h>
 #include <stdio.h>
 #include <stdlib.h>
@@ -36,14 +36,14 @@
 #endif
 
 #if __has_include(<sanitizer/asan_interface.h>)
-#	include <sanitizer/asan_interface.h>
-#	if defined(__SANITIZE_ADDRESS__)
-#		define AUTORTFM_ASAN_ENABLED 1
-#	elif defined(__has_feature)
-#		if __has_feature(address_sanitizer)
-#			define AUTORTFM_ASAN_ENABLED 1
-#		endif
-#	endif
+#include <sanitizer/asan_interface.h>
+#if defined(__SANITIZE_ADDRESS__)
+#define AUTORTFM_ASAN_ENABLED 1
+#elif defined(__has_feature)
+#if __has_feature(address_sanitizer)
+#define AUTORTFM_ASAN_ENABLED 1
+#endif
+#endif
 #endif
 
 #ifndef AUTORTFM_ASAN_ENABLED
@@ -67,19 +67,22 @@ namespace
 class AUTORTFM_INTERNAL FNullFile
 {
 public:
-    static FILE* Get()
-    {
-        static FNullFile Instance;
-        return Instance.File;
-    }
+	static FILE* Get()
+	{
+		static FNullFile Instance;
+		return Instance.File;
+	}
 
 private:
-    FNullFile() : File(fopen("/dev/null", "wb")) {}
-    ~FNullFile() { fclose(File); }
-    FILE* const File;
+	FNullFile() : File(fopen("/dev/null", "wb")) {}
+	~FNullFile()
+	{
+		fclose(File);
+	}
+	FILE* const File;
 };
 
-template<typename T>
+template <typename T>
 inline void RecordOpenWriteIfNotNull(T* Ptr)
 {
 	if (Ptr != nullptr)
@@ -90,50 +93,50 @@ inline void RecordOpenWriteIfNotNull(T* Ptr)
 
 AUTORTFM_INTERNAL void ThrowErrorFormatContainsPercentN()
 {
-    AUTORTFM_WARN("AutoRTFM does not support format strings containing '%%n'");
-    FContext* Context = FContext::Get();
-    Context->AbortByLanguageAndThrow();
+	AUTORTFM_WARN("AutoRTFM does not support format strings containing '%%n'");
+	FContext* Context = FContext::Get();
+	Context->AbortByLanguageAndThrow();
 }
 
 // Throws an error if the format string contains a '%n'.
 AUTORTFM_INTERNAL static void ThrowIfFormatContainsPercentN(const char* Format)
 {
-    for (const char* P = Format; *P != '\0'; ++P)
-    {
-        if (*P == '%')
-        {
-            switch (*++P)
-            {
-                case 'n':
-                    ThrowErrorFormatContainsPercentN();
-                    break;
-                case '\0':
-                    return;
-            }
-        }
-    }
+	for (const char* P = Format; *P != '\0'; ++P)
+	{
+		if (*P == '%')
+		{
+			switch (*++P)
+			{
+				case 'n':
+					ThrowErrorFormatContainsPercentN();
+					break;
+				case '\0':
+					return;
+			}
+		}
+	}
 }
 
 // Throws an error if the format string contains a '%n'.
 AUTORTFM_INTERNAL static void ThrowIfFormatContainsPercentN(const wchar_t* Format)
 {
-    for (const wchar_t* P = Format; *P != L'\0'; ++P)
-    {
-        if (*P == L'%')
-        {
-            switch (*++P)
-            {
-                case L'n':
-                    ThrowErrorFormatContainsPercentN();
-                    break;
-                case L'\0':
-                    return;
-            }
-        }
-    }
+	for (const wchar_t* P = Format; *P != L'\0'; ++P)
+	{
+		if (*P == L'%')
+		{
+			switch (*++P)
+			{
+				case L'n':
+					ThrowErrorFormatContainsPercentN();
+					break;
+				case L'\0':
+					return;
+			}
+		}
+	}
 }
 
-} // anonymous namespace
+}  // anonymous namespace
 
 #if AUTORTFM_PLATFORM_WINDOWS
 // vvvvvvvvvvvvvvvvvvvvvvvvvvvvvv AUTORTFM_PLATFORM_WINDOWS vvvvvvvvvvvvvvvvvvvvvvvvvvvvv
@@ -173,37 +176,43 @@ extern "C" AUTORTFM_INTERNAL UE_AUTORTFM_API long double RTFM__wcstold_l(const w
 	return _wcstold_l(String, EndPtr, Locale);
 }
 
-extern "C" AUTORTFM_INTERNAL UE_AUTORTFM_API long RTFM__wcstol_l(const wchar_t* String, wchar_t** EndPtr, int Radix, _locale_t Locale)
+extern "C" AUTORTFM_INTERNAL UE_AUTORTFM_API long RTFM__wcstol_l(
+	const wchar_t* String, wchar_t** EndPtr, int Radix, _locale_t Locale)
 {
 	RecordOpenWriteIfNotNull(EndPtr);
 	return _wcstol_l(String, EndPtr, Radix, Locale);
 }
 
-extern "C" AUTORTFM_INTERNAL UE_AUTORTFM_API long long RTFM__wcstoll_l(const wchar_t* String, wchar_t** EndPtr, int Radix, _locale_t Locale)
+extern "C" AUTORTFM_INTERNAL UE_AUTORTFM_API long long RTFM__wcstoll_l(
+	const wchar_t* String, wchar_t** EndPtr, int Radix, _locale_t Locale)
 {
 	RecordOpenWriteIfNotNull(EndPtr);
 	return _wcstoll_l(String, EndPtr, Radix, Locale);
 }
 
-extern "C" AUTORTFM_INTERNAL UE_AUTORTFM_API unsigned long RTFM__wcstoul_l(const wchar_t* String, wchar_t** EndPtr, int Radix, _locale_t Locale)
+extern "C" AUTORTFM_INTERNAL UE_AUTORTFM_API unsigned long RTFM__wcstoul_l(
+	const wchar_t* String, wchar_t** EndPtr, int Radix, _locale_t Locale)
 {
 	RecordOpenWriteIfNotNull(EndPtr);
 	return _wcstoul_l(String, EndPtr, Radix, Locale);
 }
 
-extern "C" AUTORTFM_INTERNAL UE_AUTORTFM_API unsigned long long RTFM__wcstoull_l(const wchar_t* String, wchar_t** EndPtr, int Radix, _locale_t Locale)
+extern "C" AUTORTFM_INTERNAL UE_AUTORTFM_API unsigned long long RTFM__wcstoull_l(
+	const wchar_t* String, wchar_t** EndPtr, int Radix, _locale_t Locale)
 {
 	RecordOpenWriteIfNotNull(EndPtr);
 	return _wcstoull_l(String, EndPtr, Radix, Locale);
 }
 
-extern "C" AUTORTFM_INTERNAL UE_AUTORTFM_API long long RTFM__wcstoi64_l(const wchar_t* String, wchar_t** EndPtr, int Radix, _locale_t Locale)
+extern "C" AUTORTFM_INTERNAL UE_AUTORTFM_API long long RTFM__wcstoi64_l(
+	const wchar_t* String, wchar_t** EndPtr, int Radix, _locale_t Locale)
 {
 	RecordOpenWriteIfNotNull(EndPtr);
 	return _wcstoi64_l(String, EndPtr, Radix, Locale);
 }
 
-extern "C" AUTORTFM_INTERNAL UE_AUTORTFM_API unsigned long long RTFM__wcstoui64_l(const wchar_t* String, wchar_t** EndPtr, int Radix, _locale_t Locale)
+extern "C" AUTORTFM_INTERNAL UE_AUTORTFM_API unsigned long long RTFM__wcstoui64_l(
+	const wchar_t* String, wchar_t** EndPtr, int Radix, _locale_t Locale)
 {
 	RecordOpenWriteIfNotNull(EndPtr);
 	return _wcstoui64_l(String, EndPtr, Radix, Locale);
@@ -211,113 +220,92 @@ extern "C" AUTORTFM_INTERNAL UE_AUTORTFM_API unsigned long long RTFM__wcstoui64_
 
 extern "C" AUTORTFM_INTERNAL UE_AUTORTFM_API FILE* RTFM___acrt_iob_func(unsigned Index)
 {
-    switch (Index)
-    {
-    case 1:
-    case 2:
-        return __acrt_iob_func(Index);
-    default:
+	switch (Index)
 	{
-		AUTORTFM_WARN("Attempt to get file descriptor %d (not 1 or 2) in __acrt_iob_func.", Index);
-		FContext* Context = FContext::Get();
-        Context->AbortByLanguageAndThrow();
-        return NULL;
+		case 1:
+		case 2:
+			return __acrt_iob_func(Index);
+		default:
+		{
+			AUTORTFM_WARN("Attempt to get file descriptor %d (not 1 or 2) in __acrt_iob_func.", Index);
+			FContext* Context = FContext::Get();
+			Context->AbortByLanguageAndThrow();
+			return NULL;
+		}
 	}
-    }
 }
 
 // FIXME: Does not currently support %n format specifiers.
 extern "C" AUTORTFM_INTERNAL UE_AUTORTFM_API int RTFM___stdio_common_vfprintf(
-        unsigned __int64 Options,
-        FILE*            Stream,
-        char const*      Format,
-        _locale_t        Locale,
-        va_list          ArgList)
+	unsigned __int64 Options, FILE* Stream, char const* Format, _locale_t Locale, va_list ArgList)
 {
-    ThrowIfFormatContainsPercentN(Format);
+	ThrowIfFormatContainsPercentN(Format);
 
-    return __stdio_common_vfprintf(Options, Stream, Format, Locale, ArgList);
+	return __stdio_common_vfprintf(Options, Stream, Format, Locale, ArgList);
 }
 
 // FIXME: Does not currently support %n format specifiers.
 extern "C" AUTORTFM_INTERNAL UE_AUTORTFM_API int RTFM___stdio_common_vsprintf(
-        unsigned __int64 Options,
-        char*            Buffer,
-        size_t           BufferCount,
-        char const*      Format,
-        _locale_t        Locale,
-        va_list          ArgList)
+	unsigned __int64 Options, char* Buffer, size_t BufferCount, char const* Format, _locale_t Locale, va_list ArgList)
 {
-    ThrowIfFormatContainsPercentN(Format);
+	ThrowIfFormatContainsPercentN(Format);
 
-    if (nullptr != Buffer && 0 != BufferCount)
-    {
-        va_list ArgList2;
-        va_copy(ArgList2, ArgList);
-        int Count = __stdio_common_vsprintf(Options, nullptr, 0, Format, Locale, ArgList2);
-        va_end(ArgList2);
+	if (nullptr != Buffer && 0 != BufferCount)
+	{
+		va_list ArgList2;
+		va_copy(ArgList2, ArgList);
+		int Count = __stdio_common_vsprintf(Options, nullptr, 0, Format, Locale, ArgList2);
+		va_end(ArgList2);
 
-        if (Count >= 0)
-        {
-            size_t NumBytes = std::min(BufferCount, static_cast<size_t>(1 + Count)) * sizeof(char);
-            FContext* Context = FContext::Get();
-            Context->RecordWrite(Buffer, NumBytes);
-        }
-    }
+		if (Count >= 0)
+		{
+			size_t NumBytes = std::min(BufferCount, static_cast<size_t>(1 + Count)) * sizeof(char);
+			FContext* Context = FContext::Get();
+			Context->RecordWrite(Buffer, NumBytes);
+		}
+	}
 
-    return __stdio_common_vsprintf(Options, Buffer, BufferCount, Format, Locale, ArgList);
+	return __stdio_common_vsprintf(Options, Buffer, BufferCount, Format, Locale, ArgList);
 }
 
 // FIXME: Does not currently support %n format specifiers.
 extern "C" AUTORTFM_INTERNAL UE_AUTORTFM_API int RTFM___stdio_common_vswprintf(
-        unsigned __int64 Options,
-        wchar_t*         Buffer,
-        size_t           BufferCount,
-        wchar_t const*   Format,
-        _locale_t        Locale,
-        va_list          ArgList)
+	unsigned __int64 Options, wchar_t* Buffer, size_t BufferCount, wchar_t const* Format, _locale_t Locale, va_list ArgList)
 {
-    ThrowIfFormatContainsPercentN(Format);
+	ThrowIfFormatContainsPercentN(Format);
 
-    if (nullptr != Buffer && 0 != BufferCount)
-    {
-        va_list ArgList2;
-        va_copy(ArgList2, ArgList);
-        int Count = __stdio_common_vswprintf(Options, nullptr, 0, Format, Locale, ArgList2);
-        va_end(ArgList2);
+	if (nullptr != Buffer && 0 != BufferCount)
+	{
+		va_list ArgList2;
+		va_copy(ArgList2, ArgList);
+		int Count = __stdio_common_vswprintf(Options, nullptr, 0, Format, Locale, ArgList2);
+		va_end(ArgList2);
 
-        if (Count >= 0)
-        {
-            size_t NumBytes = std::min(BufferCount, static_cast<size_t>(1 + Count)) * sizeof(wchar_t);
-            FContext* Context = FContext::Get();
-            Context->RecordWrite(Buffer, NumBytes);
-        }
-    }
+		if (Count >= 0)
+		{
+			size_t NumBytes = std::min(BufferCount, static_cast<size_t>(1 + Count)) * sizeof(wchar_t);
+			FContext* Context = FContext::Get();
+			Context->RecordWrite(Buffer, NumBytes);
+		}
+	}
 
-    return __stdio_common_vswprintf(Options, Buffer, BufferCount, Format, Locale, ArgList);
+	return __stdio_common_vswprintf(Options, Buffer, BufferCount, Format, Locale, ArgList);
 }
 
 // FIXME: Does not currently support %n format specifiers.
 extern "C" AUTORTFM_INTERNAL UE_AUTORTFM_API int RTFM___stdio_common_vfwprintf(
-        unsigned __int64 Options,
-        FILE*            Stream,
-        wchar_t const*   Format,
-        _locale_t        Locale,
-        va_list          ArgList)
+	unsigned __int64 Options, FILE* Stream, wchar_t const* Format, _locale_t Locale, va_list ArgList)
 {
-    ThrowIfFormatContainsPercentN(Format);
+	ThrowIfFormatContainsPercentN(Format);
 
-    return __stdio_common_vfwprintf(Options, Stream, Format, Locale, ArgList);
+	return __stdio_common_vfwprintf(Options, Stream, Format, Locale, ArgList);
 }
 
 extern "C" AUTORTFM_INTERNAL UE_AUTORTFM_API BOOL RTFM_TlsSetValue(DWORD dwTlsIndex, LPVOID lpTlsValue)
 {
 	LPVOID OriginalValue = TlsGetValue(dwTlsIndex);
 
-	AutoRTFM::ForTheRuntime::RegisterOnAbortFromTheOpen([dwTlsIndex, OriginalValue]
-	{
-		TlsSetValue(dwTlsIndex, OriginalValue);
-	});
+	AutoRTFM::ForTheRuntime::RegisterOnAbortFromTheOpen([dwTlsIndex, OriginalValue] { TlsSetValue(dwTlsIndex, OriginalValue); });
 
 	return TlsSetValue(dwTlsIndex, lpTlsValue);
 }
@@ -326,10 +314,7 @@ extern "C" AUTORTFM_INTERNAL UE_AUTORTFM_API BOOL RTFM_FlsSetValue(DWORD dwFlsIn
 {
 	LPVOID OriginalValue = FlsGetValue(dwFlsIndex);
 
-	AutoRTFM::ForTheRuntime::RegisterOnAbortFromTheOpen([dwFlsIndex, OriginalValue]
-	{
-		FlsSetValue(dwFlsIndex, OriginalValue);
-	});
+	AutoRTFM::ForTheRuntime::RegisterOnAbortFromTheOpen([dwFlsIndex, OriginalValue] { FlsSetValue(dwFlsIndex, OriginalValue); });
 
 	return FlsSetValue(dwFlsIndex, lpFlsValue);
 }
@@ -392,60 +377,51 @@ extern "C" AUTORTFM_INTERNAL UE_AUTORTFM_API double RTFM_strtof64(const char* St
 
 extern "C" AUTORTFM_INTERNAL UE_AUTORTFM_API void* RTFM_memcpy(void* Dst, const void* Src, size_t Size) throw()
 {
-    return Memcpy(Dst, Src, Size, FContext::Get());
+	return Memcpy(Dst, Src, Size, FContext::Get());
 }
 
 extern "C" AUTORTFM_INTERNAL UE_AUTORTFM_API void* RTFM_memmove(void* Dst, const void* Src, size_t Size) throw()
 {
-    return Memmove(Dst, Src, Size, FContext::Get());
+	return Memmove(Dst, Src, Size, FContext::Get());
 }
 
 extern "C" AUTORTFM_INTERNAL UE_AUTORTFM_API void* RTFM_memset(void* Dst, int Value, size_t Size) throw()
 {
-    return Memset(Dst, Value, Size, FContext::Get());
+	return Memset(Dst, Value, Size, FContext::Get());
 }
 
 extern "C" AUTORTFM_INTERNAL UE_AUTORTFM_API void* RTFM_malloc(size_t Size) throw()
 {
-    void* Result = malloc(Size);
+	void* Result = malloc(Size);
 	FContext* Context = FContext::Get();
-    Context->GetCurrentTransaction()->DeferUntilAbort([Result]
-    {
-        free(Result);
-    });
-    Context->DidAllocate(Result, Size);
-    return Result;
+	Context->GetCurrentTransaction()->DeferUntilAbort([Result] { free(Result); });
+	Context->DidAllocate(Result, Size);
+	return Result;
 }
 
 extern "C" AUTORTFM_INTERNAL UE_AUTORTFM_API void* RTFM_calloc(size_t Count, size_t Size) throw()
 {
-    void* Result = calloc(Count, Size);
+	void* Result = calloc(Count, Size);
 	FContext* Context = FContext::Get();
-    Context->GetCurrentTransaction()->DeferUntilAbort([Result]
-    {
-        free(Result);
-    });
-    Context->DidAllocate(Result, Count * Size);
-    return Result;
+	Context->GetCurrentTransaction()->DeferUntilAbort([Result] { free(Result); });
+	Context->DidAllocate(Result, Count * Size);
+	return Result;
 }
 
 extern "C" AUTORTFM_INTERNAL UE_AUTORTFM_API void RTFM_free(void* Ptr) throw()
 {
-    if (Ptr)
-    {
+	if (Ptr)
+	{
 		FContext* Context = FContext::Get();
-        Context->GetCurrentTransaction()->DeferUntilCommit([Ptr]
-        {
-            free(Ptr);
-        });
-    }
+		Context->GetCurrentTransaction()->DeferUntilCommit([Ptr] { free(Ptr); });
+	}
 }
 
 extern "C" AUTORTFM_INTERNAL UE_AUTORTFM_API void* RTFM_realloc(void* Ptr, size_t Size) throw()
 {
-    void* NewObject = RTFM_malloc(Size);
-    if (Ptr)
-    {
+	void* NewObject = RTFM_malloc(Size);
+	if (Ptr)
+	{
 #if defined(__APPLE__)
 		const size_t OldSize = malloc_size(Ptr);
 #elif defined(_WIN32)
@@ -454,123 +430,124 @@ extern "C" AUTORTFM_INTERNAL UE_AUTORTFM_API void* RTFM_realloc(void* Ptr, size_
 		const size_t OldSize = malloc_usable_size(Ptr);
 #endif
 		FContext* Context = FContext::Get();
-        MemcpyToNew(NewObject, Ptr, std::min(OldSize, Size), Context);
-        RTFM_free(Ptr);
-    }
-    return NewObject;
+		MemcpyToNew(NewObject, Ptr, std::min(OldSize, Size), Context);
+		RTFM_free(Ptr);
+	}
+	return NewObject;
 }
 
 extern "C" AUTORTFM_INTERNAL UE_AUTORTFM_API char* RTFM_strcpy(char* const Dst, const char* const Src) throw()
 {
-    const size_t SrcLen = strlen(Src);
+	const size_t SrcLen = strlen(Src);
 
 	FContext* Context = FContext::Get();
-    Context->RecordWrite(Dst, SrcLen + sizeof(char));
-    return strcpy(Dst, Src);
+	Context->RecordWrite(Dst, SrcLen + sizeof(char));
+	return strcpy(Dst, Src);
 }
 
 extern "C" AUTORTFM_INTERNAL UE_AUTORTFM_API char* RTFM_strncpy(char* const Dst, const char* const Src, const size_t Num) throw()
 {
 	FContext* Context = FContext::Get();
-    Context->RecordWrite(Dst, Num);
-    return strncpy(Dst, Src, Num);
+	Context->RecordWrite(Dst, Num);
+	return strncpy(Dst, Src, Num);
 }
 
 extern "C" AUTORTFM_INTERNAL UE_AUTORTFM_API char* RTFM_strcat(char* const Dst, const char* const Src) throw()
 {
-    const size_t DstLen = strlen(Dst);
-    const size_t SrcLen = strlen(Src);
+	const size_t DstLen = strlen(Dst);
+	const size_t SrcLen = strlen(Src);
 
 	FContext* Context = FContext::Get();
-    Context->RecordWrite(Dst + DstLen, SrcLen + 1);
-    return strcat(Dst, Src);
+	Context->RecordWrite(Dst + DstLen, SrcLen + 1);
+	return strcat(Dst, Src);
 }
 
 extern "C" AUTORTFM_INTERNAL UE_AUTORTFM_API char* RTFM_strncat(char* const Dst, const char* const Src, const size_t Num) throw()
 {
-    const size_t DstLen = strlen(Dst);
+	const size_t DstLen = strlen(Dst);
 
 	FContext* Context = FContext::Get();
-    Context->RecordWrite(Dst + DstLen, Num + 1);
-    return strncat(Dst, Src, Num);
+	Context->RecordWrite(Dst + DstLen, Num + 1);
+	return strncat(Dst, Src, Num);
 }
 
 extern "C" AUTORTFM_INTERNAL UE_AUTORTFM_API long RTFM_strtol(const char* String, char** EndPtr, int Radix) throw()
-{ 
-    RecordOpenWriteIfNotNull(EndPtr);
-	return strtol(String, EndPtr, Radix); 
+{
+	RecordOpenWriteIfNotNull(EndPtr);
+	return strtol(String, EndPtr, Radix);
 }
 
 extern "C" AUTORTFM_INTERNAL UE_AUTORTFM_API long long RTFM_strtoll(const char* String, char** EndPtr, int Radix) throw()
-{ 
-    RecordOpenWriteIfNotNull(EndPtr);
-	return strtoll(String, EndPtr, Radix); 
+{
+	RecordOpenWriteIfNotNull(EndPtr);
+	return strtoll(String, EndPtr, Radix);
 }
 
 extern "C" AUTORTFM_INTERNAL UE_AUTORTFM_API unsigned long RTFM_strtoul(const char* String, char** EndPtr, int Radix) throw()
-{ 
-    RecordOpenWriteIfNotNull(EndPtr);
-	return strtoul(String, EndPtr, Radix); 
+{
+	RecordOpenWriteIfNotNull(EndPtr);
+	return strtoul(String, EndPtr, Radix);
 }
 
 extern "C" AUTORTFM_INTERNAL UE_AUTORTFM_API unsigned long long RTFM_strtoull(const char* String, char** EndPtr, int Radix) throw()
-{ 
-    RecordOpenWriteIfNotNull(EndPtr);
-	return strtoull(String, EndPtr, Radix); 
+{
+	RecordOpenWriteIfNotNull(EndPtr);
+	return strtoull(String, EndPtr, Radix);
 }
 
 extern "C" AUTORTFM_INTERNAL UE_AUTORTFM_API float RTFM_strtof(const char* String, char** EndPtr) throw()
-{ 
-    RecordOpenWriteIfNotNull(EndPtr);
-	return strtof(String, EndPtr); 
+{
+	RecordOpenWriteIfNotNull(EndPtr);
+	return strtof(String, EndPtr);
 }
 
 extern "C" AUTORTFM_INTERNAL UE_AUTORTFM_API double RTFM_strtod(const char* String, char** EndPtr) throw()
-{ 
-    RecordOpenWriteIfNotNull(EndPtr);
-	return strtod(String, EndPtr); 
+{
+	RecordOpenWriteIfNotNull(EndPtr);
+	return strtod(String, EndPtr);
 }
 
 extern "C" AUTORTFM_INTERNAL UE_AUTORTFM_API double RTFM_wcstod(const wchar_t* String, wchar_t** EndPtr) throw()
-{ 
-    RecordOpenWriteIfNotNull(EndPtr);
-	return wcstod(String, EndPtr); 
+{
+	RecordOpenWriteIfNotNull(EndPtr);
+	return wcstod(String, EndPtr);
 }
 
 extern "C" AUTORTFM_INTERNAL UE_AUTORTFM_API float RTFM_wcstof(const wchar_t* String, wchar_t** EndPtr) throw()
-{ 
-    RecordOpenWriteIfNotNull(EndPtr);
-	return wcstof(String, EndPtr); 
+{
+	RecordOpenWriteIfNotNull(EndPtr);
+	return wcstof(String, EndPtr);
 }
 
 extern "C" AUTORTFM_INTERNAL UE_AUTORTFM_API long double RTFM_wcstold(const wchar_t* String, wchar_t** EndPtr) throw()
-{ 
-    RecordOpenWriteIfNotNull(EndPtr);
-	return wcstold(String, EndPtr); 
+{
+	RecordOpenWriteIfNotNull(EndPtr);
+	return wcstold(String, EndPtr);
 }
 
 extern "C" AUTORTFM_INTERNAL UE_AUTORTFM_API long RTFM_wcstol(const wchar_t* String, wchar_t** EndPtr, int Radix) throw()
-{ 
-    RecordOpenWriteIfNotNull(EndPtr);
-	return wcstol(String, EndPtr, Radix); 
+{
+	RecordOpenWriteIfNotNull(EndPtr);
+	return wcstol(String, EndPtr, Radix);
 }
 
 extern "C" AUTORTFM_INTERNAL UE_AUTORTFM_API long long RTFM_wcstoll(const wchar_t* String, wchar_t** EndPtr, int Radix) throw()
-{ 
-    RecordOpenWriteIfNotNull(EndPtr);
-	return wcstoll(String, EndPtr, Radix); 
+{
+	RecordOpenWriteIfNotNull(EndPtr);
+	return wcstoll(String, EndPtr, Radix);
 }
 
 extern "C" AUTORTFM_INTERNAL UE_AUTORTFM_API unsigned long RTFM_wcstoul(const wchar_t* String, wchar_t** EndPtr, int Radix) throw()
-{ 
-    RecordOpenWriteIfNotNull(EndPtr);
-	return wcstoul(String, EndPtr, Radix); 
+{
+	RecordOpenWriteIfNotNull(EndPtr);
+	return wcstoul(String, EndPtr, Radix);
 }
 
-extern "C" AUTORTFM_INTERNAL UE_AUTORTFM_API unsigned long long RTFM_wcstoull(const wchar_t* String, wchar_t** EndPtr, int Radix) throw()
-{ 
-    RecordOpenWriteIfNotNull(EndPtr);
-	return wcstoull(String, EndPtr, Radix); 
+extern "C" AUTORTFM_INTERNAL UE_AUTORTFM_API unsigned long long RTFM_wcstoull(
+	const wchar_t* String, wchar_t** EndPtr, int Radix) throw()
+{
+	RecordOpenWriteIfNotNull(EndPtr);
+	return wcstoull(String, EndPtr, Radix);
 }
 
 extern "C" AUTORTFM_INTERNAL UE_AUTORTFM_API std::to_chars_result RTFM__ToChars_i(char* First, char* Last, int Value, int Base)
@@ -579,7 +556,8 @@ extern "C" AUTORTFM_INTERNAL UE_AUTORTFM_API std::to_chars_result RTFM__ToChars_
 	return std::to_chars(First, Last, Value, Base);
 }
 
-extern "C" AUTORTFM_INTERNAL UE_AUTORTFM_API std::to_chars_result RTFM__ToChars_u(char* First, char* Last, unsigned int Value, int Base)
+extern "C" AUTORTFM_INTERNAL UE_AUTORTFM_API std::to_chars_result RTFM__ToChars_u(
+	char* First, char* Last, unsigned int Value, int Base)
 {
 	RecordOpenWrite(First, Last - First);
 	return std::to_chars(First, Last, Value, Base);
@@ -603,7 +581,8 @@ extern "C" AUTORTFM_INTERNAL UE_AUTORTFM_API std::to_chars_result RTFM__ToChars_
 	return std::to_chars(First, Last, Value, Base);
 }
 
-extern "C" AUTORTFM_INTERNAL UE_AUTORTFM_API std::to_chars_result RTFM__ToChars_u16(char* First, char* Last, uint16_t Value, int Base)
+extern "C" AUTORTFM_INTERNAL UE_AUTORTFM_API std::to_chars_result RTFM__ToChars_u16(
+	char* First, char* Last, uint16_t Value, int Base)
 {
 	RecordOpenWrite(First, Last - First);
 	return std::to_chars(First, Last, Value, Base);
@@ -615,7 +594,8 @@ extern "C" AUTORTFM_INTERNAL UE_AUTORTFM_API std::to_chars_result RTFM__ToChars_
 	return std::to_chars(First, Last, Value, Base);
 }
 
-extern "C" AUTORTFM_INTERNAL UE_AUTORTFM_API std::to_chars_result RTFM__ToChars_u32(char* First, char* Last, uint32_t Value, int Base)
+extern "C" AUTORTFM_INTERNAL UE_AUTORTFM_API std::to_chars_result RTFM__ToChars_u32(
+	char* First, char* Last, uint32_t Value, int Base)
 {
 	RecordOpenWrite(First, Last - First);
 	return std::to_chars(First, Last, Value, Base);
@@ -627,7 +607,8 @@ extern "C" AUTORTFM_INTERNAL UE_AUTORTFM_API std::to_chars_result RTFM__ToChars_
 	return std::to_chars(First, Last, Value, Base);
 }
 
-extern "C" AUTORTFM_INTERNAL UE_AUTORTFM_API std::to_chars_result RTFM__ToChars_u64(char* First, char* Last, uint64_t Value, int Base)
+extern "C" AUTORTFM_INTERNAL UE_AUTORTFM_API std::to_chars_result RTFM__ToChars_u64(
+	char* First, char* Last, uint64_t Value, int Base)
 {
 	RecordOpenWrite(First, Last - First);
 	return std::to_chars(First, Last, Value, Base);
@@ -639,13 +620,15 @@ extern "C" AUTORTFM_INTERNAL UE_AUTORTFM_API std::to_chars_result RTFM__ToChars_
 	return std::to_chars(First, Last, Value);
 }
 
-extern "C" AUTORTFM_INTERNAL UE_AUTORTFM_API std::to_chars_result RTFM__ToChars_f_F(char* First, char* Last, float Value, std::chars_format Format)
+extern "C" AUTORTFM_INTERNAL UE_AUTORTFM_API std::to_chars_result RTFM__ToChars_f_F(
+	char* First, char* Last, float Value, std::chars_format Format)
 {
 	RecordOpenWrite(First, Last - First);
 	return std::to_chars(First, Last, Value, Format);
 }
 
-extern "C" AUTORTFM_INTERNAL UE_AUTORTFM_API std::to_chars_result RTFM__ToChars_f_FP(char* First, char* Last, float Value, std::chars_format Format, int Precision)
+extern "C" AUTORTFM_INTERNAL UE_AUTORTFM_API std::to_chars_result RTFM__ToChars_f_FP(
+	char* First, char* Last, float Value, std::chars_format Format, int Precision)
 {
 	RecordOpenWrite(First, Last - First);
 	return std::to_chars(First, Last, Value, Format, Precision);
@@ -657,13 +640,15 @@ extern "C" AUTORTFM_INTERNAL UE_AUTORTFM_API std::to_chars_result RTFM__ToChars_
 	return std::to_chars(First, Last, Value);
 }
 
-extern "C" AUTORTFM_INTERNAL UE_AUTORTFM_API std::to_chars_result RTFM__ToChars_d_F(char* First, char* Last, double Value, std::chars_format Format)
+extern "C" AUTORTFM_INTERNAL UE_AUTORTFM_API std::to_chars_result RTFM__ToChars_d_F(
+	char* First, char* Last, double Value, std::chars_format Format)
 {
 	RecordOpenWrite(First, Last - First);
 	return std::to_chars(First, Last, Value, Format);
 }
 
-extern "C" AUTORTFM_INTERNAL UE_AUTORTFM_API std::to_chars_result RTFM__ToChars_d_FP(char* First, char* Last, double Value, std::chars_format Format, int Precision)
+extern "C" AUTORTFM_INTERNAL UE_AUTORTFM_API std::to_chars_result RTFM__ToChars_d_FP(
+	char* First, char* Last, double Value, std::chars_format Format, int Precision)
 {
 	RecordOpenWrite(First, Last - First);
 	return std::to_chars(First, Last, Value, Format, Precision);
@@ -675,13 +660,15 @@ extern "C" AUTORTFM_INTERNAL UE_AUTORTFM_API std::to_chars_result RTFM__ToChars_
 	return std::to_chars(First, Last, Value);
 }
 
-extern "C" AUTORTFM_INTERNAL UE_AUTORTFM_API std::to_chars_result RTFM__ToChars_ld_F(char* First, char* Last, long double Value, std::chars_format Format)
+extern "C" AUTORTFM_INTERNAL UE_AUTORTFM_API std::to_chars_result RTFM__ToChars_ld_F(
+	char* First, char* Last, long double Value, std::chars_format Format)
 {
 	RecordOpenWrite(First, Last - First);
 	return std::to_chars(First, Last, Value, Format);
 }
 
-extern "C" AUTORTFM_INTERNAL UE_AUTORTFM_API std::to_chars_result RTFM__ToChars_ld_FP(char* First, char* Last, long double Value, std::chars_format Format, int Precision)
+extern "C" AUTORTFM_INTERNAL UE_AUTORTFM_API std::to_chars_result RTFM__ToChars_ld_FP(
+	char* First, char* Last, long double Value, std::chars_format Format, int Precision)
 {
 	RecordOpenWrite(First, Last - First);
 	return std::to_chars(First, Last, Value, Format, Precision);
@@ -690,113 +677,112 @@ extern "C" AUTORTFM_INTERNAL UE_AUTORTFM_API std::to_chars_result RTFM__ToChars_
 // FIXME: Does not currently support %n format specifiers.
 extern "C" AUTORTFM_INTERNAL UE_AUTORTFM_API int RTFM_vsnprintf(char* Str, size_t Size, const char* Format, va_list ArgList) throw()
 {
-    ThrowIfFormatContainsPercentN(Format);
+	ThrowIfFormatContainsPercentN(Format);
 
-    if (nullptr != Str && 0 != Size)
-    {
-        va_list ArgList2;
-        va_copy(ArgList2, ArgList);
-        int Count = vsnprintf(nullptr, 0, Format, ArgList2);
-        va_end(ArgList2);
-
-        if (Count >= 0)
-        {
-            size_t NumBytes = std::min(Size, static_cast<size_t>(1 + Count)) * sizeof(char);
-            FContext* Context = FContext::Get();
-            Context->RecordWrite(Str, NumBytes);
-        }
+	if (nullptr != Str && 0 != Size)
+	{
+		va_list ArgList2;
+		va_copy(ArgList2, ArgList);
+		int Count = vsnprintf(nullptr, 0, Format, ArgList2);
+		va_end(ArgList2);
 
-    }
+		if (Count >= 0)
+		{
+			size_t NumBytes = std::min(Size, static_cast<size_t>(1 + Count)) * sizeof(char);
+			FContext* Context = FContext::Get();
+			Context->RecordWrite(Str, NumBytes);
+		}
+	}
 
-    return vsnprintf(Str, Size, Format, ArgList);
+	return vsnprintf(Str, Size, Format, ArgList);
 }
 
 // FIXME: Does not currently support %n format specifiers.
 extern "C" AUTORTFM_INTERNAL UE_AUTORTFM_API int RTFM_vswprintf(wchar_t* Str, size_t Size, const wchar_t* Format, va_list ArgList)
 {
-    ThrowIfFormatContainsPercentN(Format);
+	ThrowIfFormatContainsPercentN(Format);
 
-    if (nullptr != Str && 0 != Size)
-    {
-        va_list ArgList2;
-        va_copy(ArgList2, ArgList);
+	if (nullptr != Str && 0 != Size)
+	{
+		va_list ArgList2;
+		va_copy(ArgList2, ArgList);
 
 #if AUTORTFM_PLATFORM_WINDOWS
-        int Count = vswprintf(nullptr, 0, Format, ArgList2);
+		int Count = vswprintf(nullptr, 0, Format, ArgList2);
 #else
-        // vswprintf(nullptr, 0, ...) will return -1.
-        int Count = vfwprintf(FNullFile::Get(), Format, ArgList2);
+		// vswprintf(nullptr, 0, ...) will return -1.
+		int Count = vfwprintf(FNullFile::Get(), Format, ArgList2);
 #endif
 
-        va_end(ArgList2);
+		va_end(ArgList2);
 
-        size_t NumChars = std::min(Size, static_cast<size_t>(1 + std::max(Count, 0)));
-        size_t NumBytes = NumChars * sizeof(wchar_t);
-        if (NumBytes >= 0)
-        {
-            FContext* Context = FContext::Get();
-            Context->RecordWrite(Str, NumBytes);
-        }
-    }
+		size_t NumChars = std::min(Size, static_cast<size_t>(1 + std::max(Count, 0)));
+		size_t NumBytes = NumChars * sizeof(wchar_t);
+		if (NumBytes >= 0)
+		{
+			FContext* Context = FContext::Get();
+			Context->RecordWrite(Str, NumBytes);
+		}
+	}
 
-    return vswprintf(Str, Size, Format, ArgList);
+	return vswprintf(Str, Size, Format, ArgList);
 }
 
 // FIXME: Does not currently support %n format specifiers.
 extern "C" AUTORTFM_INTERNAL UE_AUTORTFM_API int RTFM_swprintf(wchar_t* Buffer, size_t BufferCount, wchar_t const* Format, ...)
 {
-    va_list ArgList;
+	va_list ArgList;
 
-    va_start(ArgList, Format);
-    int Count = RTFM_vswprintf(Buffer, BufferCount, Format, ArgList);
-    va_end(ArgList);
+	va_start(ArgList, Format);
+	int Count = RTFM_vswprintf(Buffer, BufferCount, Format, ArgList);
+	va_end(ArgList);
 
-    return Count;
+	return Count;
 }
 
 // FIXME: Does not currently support %n format specifiers.
 extern "C" AUTORTFM_INTERNAL UE_AUTORTFM_API int RTFM_snprintf(char* Str, size_t Size, const char* Format, ...) throw()
 {
-    va_list ArgList;
+	va_list ArgList;
 
-    va_start(ArgList, Format);
-    int Count = RTFM_vsnprintf(Str, Size, Format, ArgList);
-    va_end(ArgList);
+	va_start(ArgList, Format);
+	int Count = RTFM_vsnprintf(Str, Size, Format, ArgList);
+	va_end(ArgList);
 
-    return Count;
+	return Count;
 }
 
 // FIXME: Does not currently support %n format specifiers.
 extern "C" AUTORTFM_INTERNAL UE_AUTORTFM_API int RTFM_printf(const char* Format, ...)
 {
-    ThrowIfFormatContainsPercentN(Format);
+	ThrowIfFormatContainsPercentN(Format);
 
-    va_list ArgList;
-    va_start(ArgList, Format);
-    int Result = vprintf(Format, ArgList);
-    va_end(ArgList);
-    return Result;
+	va_list ArgList;
+	va_start(ArgList, Format);
+	int Result = vprintf(Format, ArgList);
+	va_end(ArgList);
+	return Result;
 }
 
 // FIXME: Does not currently support %n format specifiers.
 extern "C" AUTORTFM_INTERNAL UE_AUTORTFM_API int RTFM_wprintf(const wchar_t* Format, ...)
 {
-    ThrowIfFormatContainsPercentN(Format);
+	ThrowIfFormatContainsPercentN(Format);
 
-    va_list ArgList;
-    va_start(ArgList, Format);
-    int Result = vwprintf(Format, ArgList);
-    va_end(ArgList);
-    return Result;
+	va_list ArgList;
+	va_start(ArgList, Format);
+	int Result = vwprintf(Format, ArgList);
+	va_end(ArgList);
+	return Result;
 }
 
 extern "C" AUTORTFM_INTERNAL UE_AUTORTFM_API wchar_t* RTFM_wcscpy(wchar_t* Dst, const wchar_t* Src) throw()
 {
-    const size_t SrcLen = wcslen(Src);
+	const size_t SrcLen = wcslen(Src);
 
 	FContext* Context = FContext::Get();
-    Context->RecordWrite(Dst, (SrcLen + 1) * sizeof(wchar_t));
-    return wcscpy(Dst, Src);
+	Context->RecordWrite(Dst, (SrcLen + 1) * sizeof(wchar_t));
+	return wcscpy(Dst, Src);
 }
 
 extern "C" AUTORTFM_INTERNAL UE_AUTORTFM_API wchar_t* RTFM_wcsncpy(wchar_t* Dst, const wchar_t* Src, size_t Count) throw()
@@ -806,13 +792,10 @@ extern "C" AUTORTFM_INTERNAL UE_AUTORTFM_API wchar_t* RTFM_wcsncpy(wchar_t* Dst,
 	return wcsncpy(Dst, Src, Count);
 }
 
-extern "C" AUTORTFM_INTERNAL UE_AUTORTFM_API int RTFM_atexit(void(*Callback)(void)) throw()
+extern "C" AUTORTFM_INTERNAL UE_AUTORTFM_API int RTFM_atexit(void (*Callback)(void)) throw()
 {
 	FContext* Context = FContext::Get();
-	Context->GetCurrentTransaction()->DeferUntilCommit([Callback]
-		{
-			atexit(Callback);
-		});
+	Context->GetCurrentTransaction()->DeferUntilCommit([Callback] { atexit(Callback); });
 
 	return 0;
 }
@@ -822,35 +805,29 @@ extern "C" void __stdcall _CxxThrowException(void* pExceptionObject, struct _Thr
 #endif
 
 #if AUTORTFM_PLATFORM_LINUX && AUTORTFM_EXCEPTIONS_ENABLED
-extern "C" void __cxa_throw(void* thrown_exception, struct std::type_info * tinfo, void (*dest)(void*));
+extern "C" void __cxa_throw(void* thrown_exception, struct std::type_info* tinfo, void (*dest)(void*));
 extern "C" void* __cxa_allocate_exception(size_t thrown_size) throw();
 extern "C" void* __cxa_begin_catch(void* exceptionObject) throw();
 extern "C" void __cxa_end_catch();
 #endif
 
-template<typename OpenNewFn, typename OpenDeleteFn, typename ... ArgTys>
+template <typename OpenNewFn, typename OpenDeleteFn, typename... ArgTys>
 UE_AUTORTFM_FORCEINLINE static void* RTFM_ClosedNew(OpenNewFn* New, OpenDeleteFn* Delete, unsigned long Size, ArgTys... Args)
 {
 	void* Result = New(Size, Args...);
 	FContext* Context = FContext::Get();
-	Context->GetCurrentTransaction()->DeferUntilAbort([=]
-	{
-		Delete(Result, Args...);
-	});
+	Context->GetCurrentTransaction()->DeferUntilAbort([=] { Delete(Result, Args...); });
 	Context->DidAllocate(Result, Size);
 	return Result;
 }
 
-template<typename OpenDeleteFn, typename ... ArgTys>
+template <typename OpenDeleteFn, typename... ArgTys>
 UE_AUTORTFM_FORCEINLINE static void RTFM_ClosedDelete(OpenDeleteFn* Delete, void* Pointer, ArgTys... Args)
 {
 	if (Pointer)
 	{
 		FContext* Context = FContext::Get();
-		Context->GetCurrentTransaction()->DeferUntilCommit([=]
-		{
-			Delete(Pointer, Args...);
-		});
+		Context->GetCurrentTransaction()->DeferUntilCommit([=] { Delete(Pointer, Args...); });
 	}
 }
 
@@ -886,7 +863,7 @@ extern "C"
 	void _ZdaPvSt11align_val_tRKSt9nothrow_t(void*, std::align_val_t, const std::nothrow_t&);
 	void _ZdlPvmSt11align_val_t(void*, unsigned long, std::align_val_t);
 	void _ZdaPvmSt11align_val_t(void*, unsigned long, std::align_val_t);
-} // extern "C"
+}  // extern "C"
 
 extern "C" AUTORTFM_INTERNAL UE_AUTORTFM_API void* RTFM_Znwm(unsigned long Size)
 {
@@ -918,12 +895,14 @@ extern "C" AUTORTFM_INTERNAL UE_AUTORTFM_API void* RTFM_ZnamSt11align_val_t(unsi
 	return RTFM_ClosedNew(_ZnamSt11align_val_t, _ZdaPvSt11align_val_t, Size, Align);
 }
 
-extern "C" AUTORTFM_INTERNAL UE_AUTORTFM_API void* RTFM_ZnwmSt11align_val_tRKSt9nothrow_t(unsigned long Size, std::align_val_t Align, const std::nothrow_t& NoThrow)
+extern "C" AUTORTFM_INTERNAL UE_AUTORTFM_API void* RTFM_ZnwmSt11align_val_tRKSt9nothrow_t(
+	unsigned long Size, std::align_val_t Align, const std::nothrow_t& NoThrow)
 {
 	return RTFM_ClosedNew(_ZnwmSt11align_val_tRKSt9nothrow_t, _ZdlPvSt11align_val_tRKSt9nothrow_t, Size, Align, NoThrow);
 }
 
-extern "C" AUTORTFM_INTERNAL UE_AUTORTFM_API void* RTFM_ZnamSt11align_val_tRKSt9nothrow_t(unsigned long Size, std::align_val_t Align, const std::nothrow_t& NoThrow)
+extern "C" AUTORTFM_INTERNAL UE_AUTORTFM_API void* RTFM_ZnamSt11align_val_tRKSt9nothrow_t(
+	unsigned long Size, std::align_val_t Align, const std::nothrow_t& NoThrow)
 {
 	return RTFM_ClosedNew(_ZnamSt11align_val_tRKSt9nothrow_t, _ZdaPvSt11align_val_tRKSt9nothrow_t, Size, Align, NoThrow);
 }
@@ -968,28 +947,32 @@ extern "C" AUTORTFM_INTERNAL UE_AUTORTFM_API void RTFM_ZdaPvSt11align_val_t(void
 	RTFM_ClosedDelete(_ZdaPvSt11align_val_t, Pointer, Align);
 }
 
-extern "C" AUTORTFM_INTERNAL UE_AUTORTFM_API void RTFM_ZdlPvSt11align_val_tRKSt9nothrow_t(void* Pointer, std::align_val_t Align, const std::nothrow_t& NoThrow)
+extern "C" AUTORTFM_INTERNAL UE_AUTORTFM_API void RTFM_ZdlPvSt11align_val_tRKSt9nothrow_t(
+	void* Pointer, std::align_val_t Align, const std::nothrow_t& NoThrow)
 {
 	RTFM_ClosedDelete(_ZdlPvSt11align_val_tRKSt9nothrow_t, Pointer, Align, NoThrow);
 }
 
-extern "C" AUTORTFM_INTERNAL UE_AUTORTFM_API void RTFM_ZdaPvSt11align_val_tRKSt9nothrow_t(void* Pointer, std::align_val_t Align, const std::nothrow_t& NoThrow)
+extern "C" AUTORTFM_INTERNAL UE_AUTORTFM_API void RTFM_ZdaPvSt11align_val_tRKSt9nothrow_t(
+	void* Pointer, std::align_val_t Align, const std::nothrow_t& NoThrow)
 {
 	RTFM_ClosedDelete(_ZdaPvSt11align_val_tRKSt9nothrow_t, Pointer, Align, NoThrow);
 }
 
-extern "C" AUTORTFM_INTERNAL UE_AUTORTFM_API void RTFM_ZdlPvmSt11align_val_t(void* Pointer, unsigned long Size, std::align_val_t Align)
+extern "C" AUTORTFM_INTERNAL UE_AUTORTFM_API void RTFM_ZdlPvmSt11align_val_t(
+	void* Pointer, unsigned long Size, std::align_val_t Align)
 {
 	RTFM_ClosedDelete(_ZdlPvmSt11align_val_t, Pointer, Size, Align);
 }
 
-extern "C" AUTORTFM_INTERNAL UE_AUTORTFM_API void RTFM_ZdaPvmSt11align_val_t(void* Pointer, unsigned long Size, std::align_val_t Align)
+extern "C" AUTORTFM_INTERNAL UE_AUTORTFM_API void RTFM_ZdaPvmSt11align_val_t(
+	void* Pointer, unsigned long Size, std::align_val_t Align)
 {
 	RTFM_ClosedDelete(_ZdaPvmSt11align_val_t, Pointer, Size, Align);
 }
 // ^^^^^^^^^^^^^^^^^^^^^^^^^^^^^^ AUTORTFM_PLATFORM_LINUX ^^^^^^^^^^^^^^^^^^^^^^^^^^^^^^
 #endif
 
-} // namespace AutoRTFM
+}  // namespace AutoRTFM
 
-#endif // defined(__AUTORTFM) && __AUTORTFM
+#endif  // defined(__AUTORTFM) && __AUTORTFM
diff --git a/Engine/Source/Runtime/AutoRTFM/Private/StlAllocator.h b/Engine/Source/Runtime/AutoRTFM/Private/StlAllocator.h
--- a/Engine/Source/Runtime/AutoRTFM/Private/StlAllocator.h
+++ b/Engine/Source/Runtime/AutoRTFM/Private/StlAllocator.h
@@ -24,7 +24,7 @@ public:
 	StlAllocator() = default;
 
 	template <typename U>
-	StlAllocator(const StlAllocator<U>&) {};
+	StlAllocator(const StlAllocator<U>&){};
 
 	T* allocate(size_t Count)
 	{
@@ -49,6 +49,6 @@ bool operator!=(const StlAllocator<T>&, const StlAllocator<U>&)
 	return false;
 }
 
-}  // AutoRTFM
+}  // namespace AutoRTFM
 
-#endif // (defined(__AUTORTFM) && __AUTORTFM)
+#endif  // (defined(__AUTORTFM) && __AUTORTFM)
diff --git a/Engine/Source/Runtime/AutoRTFM/Private/Transaction.h b/Engine/Source/Runtime/AutoRTFM/Private/Transaction.h
--- a/Engine/Source/Runtime/AutoRTFM/Private/Transaction.h
+++ b/Engine/Source/Runtime/AutoRTFM/Private/Transaction.h
@@ -17,7 +17,7 @@
 
 #if AUTORTFM_SANITIZER
 #include "CallstackTree.h"
-#endif // AUTORTFM_SANITIZER
+#endif  // AUTORTFM_SANITIZER
 
 #include <optional>
 
@@ -45,13 +45,19 @@ public:
 	// Clears the tracked transaction state and resets back to the default Uninitialized state.
 	void Reset();
 
-	bool IsNested() const { return !!Parent; }
-	FTransaction* GetParent() const { return Parent; }
+	bool IsNested() const
+	{
+		return !!Parent;
+	}
+	FTransaction* GetParent() const
+	{
+		return Parent;
+	}
 
 	// This should just use type displays or ranges. Maybe ranges could even work out great.
 	bool IsNestedWithin(const FTransaction* Other) const
 	{
-		for (const FTransaction* Current = this; ; Current = Current->Parent)
+		for (const FTransaction* Current = this;; Current = Current->Parent)
 		{
 			if (!Current)
 			{
@@ -64,7 +70,10 @@ public:
 		}
 	}
 
-	inline bool IsScopedTransaction() const { return bIsStackScoped; }
+	inline bool IsScopedTransaction() const
+	{
+		return bIsStackScoped;
+	}
 
 	void DeferUntilCommit(TTask<void()>&&);
 	void DeferUntilPreAbort(TTask<void()>&&);
@@ -83,9 +92,12 @@ public:
 
 	// Record that a write is about to occur at the given LogicalAddress of Size bytes.
 	void RecordWrite(void* LogicalAddress, size_t Size, EWriteFlags Flags = EWriteFlags::Default);
-	template<unsigned SIZE> void RecordWrite(void* LogicalAddress);
-	template<unsigned SIZE> void RecordWriteInsertedSlow(void* LogicalAddress);
-	template<unsigned SIZE> void RecordWriteNotInsertedSlow(void* LogicalAddress);
+	template <unsigned SIZE>
+	void RecordWrite(void* LogicalAddress);
+	template <unsigned SIZE>
+	void RecordWriteInsertedSlow(void* LogicalAddress);
+	template <unsigned SIZE>
+	void RecordWriteNotInsertedSlow(void* LogicalAddress);
 
 	void DidAllocate(void* LogicalAddress, size_t Size);
 	void DidFree(void* LogicalAddress);
@@ -94,7 +106,10 @@ public:
 	bool IsFresh() const;
 
 	// The stack range represents all stack memory inside the transaction scope
-	inline FStackRange GetStackRange() const { return StackRange; }
+	inline FStackRange GetStackRange() const
+	{
+		return StackRange;
+	}
 
 	// Returns true if the LogicalAddress is within the stack of the transaction.
 	inline bool IsOnStack(const void* LogicalAddress) const;
@@ -102,11 +117,14 @@ public:
 #if AUTORTFM_SANITIZER
 	inline void SanitizerOpenWrite(void* Ptr, size_t Size);
 
-	inline void DisableSanitizer() { SanitizerDisableCount++; }
+	inline void DisableSanitizer()
+	{
+		SanitizerDisableCount++;
+	}
 
 	inline void EnableSanitizer()
 	{
-		AUTORTFM_ASSERT(SanitizerDisableCount > 0);;
+		AUTORTFM_ASSERT(SanitizerDisableCount > 0);
 		SanitizerDisableCount--;
 	}
 
@@ -146,7 +164,7 @@ private:
 	// - When the transaction is active, the Parent field is used to support transaction nesting.
 	//   This field points to our immediate outer nest; it's null when this is the top level.
 	// - When the transaction is "freed"--sitting in the FTransactionPool free list--this field is
-	//   used to point to the next free item, forming a linked list of reusable FTransactions. 
+	//   used to point to the next free item, forming a linked list of reusable FTransactions.
 	//   (See GetIntrusiveAddress.)
 	FTransaction* Parent;
 
@@ -192,6 +210,6 @@ private:
 	bool bIsInAllocateFn;
 };
 
-} // namespace AutoRTFM
+}  // namespace AutoRTFM
 
-#endif // (defined(__AUTORTFM) && __AUTORTFM)
+#endif  // (defined(__AUTORTFM) && __AUTORTFM)
diff --git a/Engine/Source/Runtime/AutoRTFM/Private/Utils.h b/Engine/Source/Runtime/AutoRTFM/Private/Utils.h
--- a/Engine/Source/Runtime/AutoRTFM/Private/Utils.h
+++ b/Engine/Source/Runtime/AutoRTFM/Private/Utils.h
@@ -34,32 +34,30 @@ namespace AutoRTFM
 // based on ForTheRuntime::GetInternalAbortAction() and
 // ForTheRuntime::GetEnsureOnInternalAbort().
 // If ProgramCounter is null, then the return address of the call will be used.
-AUTORTFM_DISABLE __attribute__((__format__(__printf__, 4, 5)))
-UE_AUTORTFM_API void ReportError(const char* File, int Line, void* ProgramCounter, const char* Format, ...);
+AUTORTFM_DISABLE __attribute__((__format__(__printf__, 4, 5))) UE_AUTORTFM_API void ReportError(
+	const char* File, int Line, void* ProgramCounter, const char* Format, ...);
 
-[[noreturn]] inline void InternalUnreachable(); //-V1082 Silence PVS [[noreturn]] false positive warning
+[[noreturn]] inline void InternalUnreachable();  //-V1082 Silence PVS [[noreturn]] false positive warning
 
 std::string GetFunctionDescription(void* FunctionPtr);
 
-template<typename TReturnType, typename... TParameterTypes>
+template <typename TReturnType, typename... TParameterTypes>
 std::string GetFunctionDescription(TReturnType (*FunctionPtr)(TParameterTypes...))
 {
-    return GetFunctionDescription(reinterpret_cast<void*>(FunctionPtr));
+	return GetFunctionDescription(reinterpret_cast<void*>(FunctionPtr));
 }
 
 // We use a special calling convention for this call to minimize
 // the cost of calling this function (which is super unlikely to
 // be called because it is the asserting true codepath, EG. only
 // fatal process destroying explosions happen here).
-[[clang::preserve_most, noreturn]]
-UE_AUTORTFM_FORCENOINLINE void DoAssert(void(*Logger)());
+[[clang::preserve_most, noreturn]] UE_AUTORTFM_FORCENOINLINE void DoAssert(void (*Logger)());
 
 // We use a special calling convention for this call to minimize
 // the cost of calling this function (which is super unlikely to
 // be called because it is the expecting true codepath, EG. only
 // things we don't expect to happen get here).
-[[clang::preserve_most]]
-UE_AUTORTFM_FORCENOINLINE void DoExpect(void(*Logger)());
+[[clang::preserve_most]] UE_AUTORTFM_FORCENOINLINE void DoExpect(void (*Logger)());
 
 // Logs a message using a printf-style format string and va_list arguments.
 inline void LogV(const char* File, int Line, void* ProgramCounter, autortfm_log_severity Severity, const char* Format, va_list Args)
@@ -73,8 +71,8 @@ inline void LogV(const char* File, int Line, void* ProgramCounter, autortfm_log_
 }
 
 // Logs a message using a printf-style format string and variadic arguments.
-__attribute__((__format__(__printf__, 5, 6)))
-inline void Log(const char* File, int Line, void* ProgramCounter, autortfm_log_severity Severity, const char* Format, ...)
+__attribute__((__format__(__printf__, 5, 6))) inline void Log(
+	const char* File, int Line, void* ProgramCounter, autortfm_log_severity Severity, const char* Format, ...)
 {
 	if (ProgramCounter == nullptr)
 	{
@@ -94,8 +92,7 @@ inline void LogWithCallstackV(autortfm_log_severity Severity, const char* Format
 }
 
 // Logs a message with a callstack using a printf-style format string and variadic arguments.
-__attribute__((__format__(__printf__, 2, 3)))
-inline void LogWithCallstack(autortfm_log_severity Severity, const char* Format, ...)
+__attribute__((__format__(__printf__, 2, 3))) inline void LogWithCallstack(autortfm_log_severity Severity, const char* Format, ...)
 {
 	va_list Args;
 	va_start(Args, Format);
@@ -117,8 +114,8 @@ inline void EnsureFailureV(const char* File, int Line, void* ProgramCounter, con
 
 // Reports an ensure failure using a printf-style format string and variadic arguments.
 // If ProgramCounter is null, then the return address of the call will be used.
-__attribute__((__format__(__printf__, 5, 6)))
-inline void EnsureFailure(const char* File, int Line, void* ProgramCounter, const char* Condition, const char* Format, ...)
+__attribute__((__format__(__printf__, 5, 6))) inline void EnsureFailure(
+	const char* File, int Line, void* ProgramCounter, const char* Condition, const char* Format, ...)
 {
 	if (ProgramCounter == nullptr)
 	{
@@ -133,7 +130,7 @@ inline void EnsureFailure(const char* File, int Line, void* ProgramCounter, cons
 
 // Rounds Value down to the next multiple of Alignment.
 // Alignment must be a power of two.
-template<typename T>
+template <typename T>
 inline constexpr T AlignDown(T Value, T Alignment)
 {
 	static_assert(std::is_integral_v<T>);
@@ -142,7 +139,7 @@ inline constexpr T AlignDown(T Value, T Alignment)
 
 // Rounds Value up to the next multiple of Alignment.
 // Alignment must be a power of two.
-template<typename T>
+template <typename T>
 inline constexpr T AlignUp(T Value, T Alignment)
 {
 	static_assert(std::is_integral_v<T>);
@@ -151,7 +148,7 @@ inline constexpr T AlignUp(T Value, T Alignment)
 
 // Rounds Value down to the next multiple of Multiple.
 // Unlike AlignDown(), Multiple does not need to be a power of two.
-template<typename T>
+template <typename T>
 inline constexpr T RoundDown(T Value, T Multiple)
 {
 	static_assert(std::is_integral_v<T>);
@@ -160,7 +157,7 @@ inline constexpr T RoundDown(T Value, T Multiple)
 
 // Rounds Value up to the next multiple of Multiple.
 // Unlike AlignUp(), Multiple does not need to be a power of two.
-template<typename T>
+template <typename T>
 inline constexpr T RoundUp(T Value, T Multiple)
 {
 	static_assert(std::is_integral_v<T>);
@@ -168,22 +165,24 @@ inline constexpr T RoundUp(T Value, T Multiple)
 }
 
 // Returns the linear interpolation between Start and End using the coefficient Fraction.
-template<typename T>
+template <typename T>
 inline constexpr T Lerp(T Start, T End, T Fraction)
 {
 	return Start + Fraction * (End - Start);
 }
 
-} // namespace AutoRTFM
+}  // namespace AutoRTFM
 
 #define AUTORTFM_LIKELY(x) __builtin_expect(!!(x), 1)
 #define AUTORTFM_UNLIKELY(x) __builtin_expect(!!(x), 0)
 
 #define AUTORTFM_REQUIRE_SEMICOLON static_assert(true)
 
-#define AUTORTFM_REPORT_ERROR(Format, ...) ::AutoRTFM::ReportError(__FILE__, __LINE__, /* ProgramCounter */ nullptr, Format, ##__VA_ARGS__)
+#define AUTORTFM_REPORT_ERROR(Format, ...) \
+	::AutoRTFM::ReportError(__FILE__, __LINE__, /* ProgramCounter */ nullptr, Format, ##__VA_ARGS__)
 
-#define AUTORTFM_LOG_AT_SEVERITY(Severity, Format, ...) ::AutoRTFM::Log(__FILE__, __LINE__, /* ProgramCounter */ nullptr, Severity, Format, ##__VA_ARGS__)
+#define AUTORTFM_LOG_AT_SEVERITY(Severity, Format, ...) \
+	::AutoRTFM::Log(__FILE__, __LINE__, /* ProgramCounter */ nullptr, Severity, Format, ##__VA_ARGS__)
 
 #define AUTORTFM_VERBOSE(Format, ...) AUTORTFM_LOG_AT_SEVERITY(autortfm_log_verbose, Format, ##__VA_ARGS__)
 #define AUTORTFM_LOG(Format, ...) AUTORTFM_LOG_AT_SEVERITY(autortfm_log_info, Format, ##__VA_ARGS__)
@@ -191,17 +190,57 @@ inline constexpr T Lerp(T Start, T End, T Fraction)
 #define AUTORTFM_ERROR(Format, ...) AUTORTFM_LOG_AT_SEVERITY(autortfm_log_error, Format, ##__VA_ARGS__)
 #define AUTORTFM_FATAL(Format, ...) AUTORTFM_LOG_AT_SEVERITY(autortfm_log_fatal, Format, ##__VA_ARGS__)
 
-#define AUTORTFM_VERBOSE_V(Format, Args) ::AutoRTFM::LogV(__FILE__, __LINE__, /* ProgramCounter */ nullptr, autortfm_log_verbose, Format, Args)
-#define AUTORTFM_LOG_V(Format, Args) ::AutoRTFM::LogV(__FILE__, __LINE__, /* ProgramCounter */ nullptr, autortfm_log_info, Format, Args)
-#define AUTORTFM_WARN_V(Format, Args) ::AutoRTFM::LogV(__FILE__, __LINE__, /* ProgramCounter */ nullptr, autortfm_log_warn, Format, Args)
-#define AUTORTFM_ERROR_V(Format, Args) ::AutoRTFM::LogV(__FILE__, __LINE__, /* ProgramCounter */ nullptr, autortfm_log_error, Format, Args)
-#define AUTORTFM_FATAL_V(Format, Args) ::AutoRTFM::LogV(__FILE__, __LINE__, /* ProgramCounter */ nullptr, autortfm_log_fatal, Format, Args)
-
-#define AUTORTFM_VERBOSE_IF(Condition, Format, ...) do { if (Condition) { AUTORTFM_VERBOSE(Format, ##__VA_ARGS__); } } while(false)
-#define AUTORTFM_LOG_IF(Condition, Format, ...) do { if (Condition) { AUTORTFM_LOG(Format, ##__VA_ARGS__); } } while(false)
-#define AUTORTFM_WARN_IF(Condition, Format, ...) do { if (Condition) { AUTORTFM_WARN(Format, ##__VA_ARGS__); } } while(false)
-#define AUTORTFM_ERROR_IF(Condition, Format, ...) do { if (AUTORTFM_UNLIKELY(Condition)) { AUTORTFM_ERROR(Format, ##__VA_ARGS__); } } while(false)
-#define AUTORTFM_FATAL_IF(Condition, Format, ...) do { if (AUTORTFM_UNLIKELY(Condition)) { AUTORTFM_FATAL(Format, ##__VA_ARGS__); } } while(false)
+#define AUTORTFM_VERBOSE_V(Format, Args) \
+	::AutoRTFM::LogV(__FILE__, __LINE__, /* ProgramCounter */ nullptr, autortfm_log_verbose, Format, Args)
+#define AUTORTFM_LOG_V(Format, Args) \
+	::AutoRTFM::LogV(__FILE__, __LINE__, /* ProgramCounter */ nullptr, autortfm_log_info, Format, Args)
+#define AUTORTFM_WARN_V(Format, Args) \
+	::AutoRTFM::LogV(__FILE__, __LINE__, /* ProgramCounter */ nullptr, autortfm_log_warn, Format, Args)
+#define AUTORTFM_ERROR_V(Format, Args) \
+	::AutoRTFM::LogV(__FILE__, __LINE__, /* ProgramCounter */ nullptr, autortfm_log_error, Format, Args)
+#define AUTORTFM_FATAL_V(Format, Args) \
+	::AutoRTFM::LogV(__FILE__, __LINE__, /* ProgramCounter */ nullptr, autortfm_log_fatal, Format, Args)
+
+#define AUTORTFM_VERBOSE_IF(Condition, Format, ...)  \
+	do                                               \
+	{                                                \
+		if (Condition)                               \
+		{                                            \
+			AUTORTFM_VERBOSE(Format, ##__VA_ARGS__); \
+		}                                            \
+	} while (false)
+#define AUTORTFM_LOG_IF(Condition, Format, ...)  \
+	do                                           \
+	{                                            \
+		if (Condition)                           \
+		{                                        \
+			AUTORTFM_LOG(Format, ##__VA_ARGS__); \
+		}                                        \
+	} while (false)
+#define AUTORTFM_WARN_IF(Condition, Format, ...)  \
+	do                                            \
+	{                                             \
+		if (Condition)                            \
+		{                                         \
+			AUTORTFM_WARN(Format, ##__VA_ARGS__); \
+		}                                         \
+	} while (false)
+#define AUTORTFM_ERROR_IF(Condition, Format, ...)  \
+	do                                             \
+	{                                              \
+		if (AUTORTFM_UNLIKELY(Condition))          \
+		{                                          \
+			AUTORTFM_ERROR(Format, ##__VA_ARGS__); \
+		}                                          \
+	} while (false)
+#define AUTORTFM_FATAL_IF(Condition, Format, ...)  \
+	do                                             \
+	{                                              \
+		if (AUTORTFM_UNLIKELY(Condition))          \
+		{                                          \
+			AUTORTFM_FATAL(Format, ##__VA_ARGS__); \
+		}                                          \
+	} while (false)
 
 // If AUTORTFM_VERBOSE_ENABLED is 0, then replace all the verbose logging macros
 // with no-ops.
@@ -214,29 +253,44 @@ inline constexpr T Lerp(T Start, T End, T Fraction)
 #define AUTORTFM_VERBOSE_IF(...) AUTORTFM_REQUIRE_SEMICOLON
 #endif
 
-#define AUTORTFM_ENSURE(Condition) do { \
-	if (AUTORTFM_UNLIKELY(!(Condition))) { \
-		[[maybe_unused]] static bool bCalled = [] { \
-			::AutoRTFM::EnsureFailure(__FILE__, __LINE__, /* ProgramCounter */ nullptr, #Condition, nullptr); \
-			return true; \
-		}(); \
-	} } while(false)
-
-#define AUTORTFM_ENSURE_MSG(Condition, Format, ...) do { \
-	if (AUTORTFM_UNLIKELY(!(Condition))) { \
-		[[maybe_unused]] static bool bCalled = [&] { \
-			::AutoRTFM::EnsureFailure(__FILE__, __LINE__, /* ProgramCounter */ nullptr, #Condition, Format, ##__VA_ARGS__); \
-			return true; \
-		}(); \
-	} } while(false)
-
-#define AUTORTFM_ENSURE_MSG_V(Condition, Format, Args) do { \
-	if (AUTORTFM_UNLIKELY(!(Condition))) { \
-		[[maybe_unused]] static bool bCalled = [&] { \
-			::AutoRTFM::EnsureFailureV(__FILE__, __LINE__, /* ProgramCounter */ nullptr, #Condition, Format, Args); \
-			return true; \
-		}(); \
-	} } while(false)
+#define AUTORTFM_ENSURE(Condition)                                                                                \
+	do                                                                                                            \
+	{                                                                                                             \
+		if (AUTORTFM_UNLIKELY(!(Condition)))                                                                      \
+		{                                                                                                         \
+			[[maybe_unused]] static bool bCalled = []                                                             \
+			{                                                                                                     \
+				::AutoRTFM::EnsureFailure(__FILE__, __LINE__, /* ProgramCounter */ nullptr, #Condition, nullptr); \
+				return true;                                                                                      \
+			}();                                                                                                  \
+		}                                                                                                         \
+	} while (false)
+
+#define AUTORTFM_ENSURE_MSG(Condition, Format, ...)                                                                             \
+	do                                                                                                                          \
+	{                                                                                                                           \
+		if (AUTORTFM_UNLIKELY(!(Condition)))                                                                                    \
+		{                                                                                                                       \
+			[[maybe_unused]] static bool bCalled = [&]                                                                          \
+			{                                                                                                                   \
+				::AutoRTFM::EnsureFailure(__FILE__, __LINE__, /* ProgramCounter */ nullptr, #Condition, Format, ##__VA_ARGS__); \
+				return true;                                                                                                    \
+			}();                                                                                                                \
+		}                                                                                                                       \
+	} while (false)
+
+#define AUTORTFM_ENSURE_MSG_V(Condition, Format, Args)                                                                  \
+	do                                                                                                                  \
+	{                                                                                                                   \
+		if (AUTORTFM_UNLIKELY(!(Condition)))                                                                            \
+		{                                                                                                               \
+			[[maybe_unused]] static bool bCalled = [&]                                                                  \
+			{                                                                                                           \
+				::AutoRTFM::EnsureFailureV(__FILE__, __LINE__, /* ProgramCounter */ nullptr, #Condition, Format, Args); \
+				return true;                                                                                            \
+			}();                                                                                                        \
+		}                                                                                                               \
+	} while (false)
 
 #define AUTORTFM_DUMP_STACKTRACE(Message, ...) ::AutoRTFM::LogWithCallstack(autortfm_log_error, Message, ##__VA_ARGS__)
 
@@ -252,14 +306,14 @@ inline constexpr T Lerp(T Start, T End, T Fraction)
 // and line numbers), but then pass this to another function
 // `DoAssert` which has a special calling convention that makes
 // the caller as optimal as possible (at the expense of the callee).
-#define AUTORTFM_ASSERT(exp) \
-	if (AUTORTFM_UNLIKELY(!(exp))) \
-	{ \
-		auto Lambda = [] \
-		{ \
+#define AUTORTFM_ASSERT(exp)                                     \
+	if (AUTORTFM_UNLIKELY(!(exp)))                               \
+	{                                                            \
+		auto Lambda = []                                         \
+		{                                                        \
 			AUTORTFM_FATAL("AUTORTFM_ASSERT(%s) failure", #exp); \
-		}; \
-		AutoRTFM::DoAssert(Lambda); \
+		};                                                       \
+		AutoRTFM::DoAssert(Lambda);                              \
 	}
 
 // For places where zero cost is truly important, AUTORTFM_ASSERT_DEBUG will
@@ -269,27 +323,27 @@ inline constexpr T Lerp(T Start, T End, T Fraction)
 #if AUTORTFM_BUILD_SHIPPING || AUTORTFM_BUILD_TEST
 #define AUTORTFM_ASSERT_DEBUG(exp) ((void)AUTORTFM_LIKELY(static_cast<bool>(exp)))
 #else
-#define AUTORTFM_ASSERT_DEBUG(exp) \
-	if (AUTORTFM_UNLIKELY(!(exp))) \
-	{ \
-		auto Lambda = [] \
-		{ \
+#define AUTORTFM_ASSERT_DEBUG(exp)                                     \
+	if (AUTORTFM_UNLIKELY(!(exp)))                                     \
+	{                                                                  \
+		auto Lambda = []                                               \
+		{                                                              \
 			AUTORTFM_FATAL("AUTORTFM_ASSERT_DEBUG(%s) failure", #exp); \
-		}; \
-		AutoRTFM::DoAssert(Lambda); \
+		};                                                             \
+		AutoRTFM::DoAssert(Lambda);                                    \
 	}
 #endif
 
 // Same funkiness as AUTORTFM_ASSERT, except that it doesn't cause a fatal error.
 // TODO: Maybe upgrade the Info -> Warning?
-#define AUTORTFM_EXPECT(exp) \
-	if (AUTORTFM_UNLIKELY(!(exp))) \
-	{ \
-		auto Lambda = [] \
-		{ \
+#define AUTORTFM_EXPECT(exp)                                                                      \
+	if (AUTORTFM_UNLIKELY(!(exp)))                                                                \
+	{                                                                                             \
+		auto Lambda = []                                                                          \
+		{                                                                                         \
 			::AutoRTFM::LogWithCallstack(autortfm_log_info, "AUTORTFM_EXPECT(%s) failure", #exp); \
-		}; \
-		AutoRTFM::DoExpect(Lambda); \
+		};                                                                                        \
+		AutoRTFM::DoExpect(Lambda);                                                               \
 	}
 
 #if defined(__has_feature)
@@ -311,13 +365,13 @@ inline constexpr T Lerp(T Start, T End, T Fraction)
 #endif
 
 #if !defined(CA_SUPPRESS)
-#define CA_SUPPRESS( WarningNumber )
+#define CA_SUPPRESS(WarningNumber)
 #endif
 
-[[noreturn]] inline void AutoRTFM::InternalUnreachable() //-V1082 Silence PVS [[noreturn]] false positive warning
+[[noreturn]] inline void AutoRTFM::InternalUnreachable()  //-V1082 Silence PVS [[noreturn]] false positive warning
 {
 	AUTORTFM_FATAL("Unreachable encountered!");
 	__builtin_unreachable();
 }
 
-#endif // (defined(__AUTORTFM) && __AUTORTFM)
+#endif  // (defined(__AUTORTFM) && __AUTORTFM)
diff --git a/Engine/Source/Runtime/AutoRTFM/Private/WriteLog.h b/Engine/Source/Runtime/AutoRTFM/Private/WriteLog.h
--- a/Engine/Source/Runtime/AutoRTFM/Private/WriteLog.h
+++ b/Engine/Source/Runtime/AutoRTFM/Private/WriteLog.h
@@ -16,668 +16,685 @@
 
 namespace AutoRTFM
 {
-	struct AUTORTFM_DISABLE_NO_SANITIZE FWriteLogEntry final
+struct AUTORTFM_DISABLE_NO_SANITIZE FWriteLogEntry final
+{
+	// The address of the write.
+	std::byte* LogicalAddress = nullptr;
+
+	// A pointer to the original data before the write occurred.
+	std::byte* Data = nullptr;
+
+	// The size of the write in bytes.
+	size_t Size = 0;
+
+	// Flags on the write.
+	EWriteFlags Flags = EWriteFlags::Default;
+
+	// Returns a new FWriteLogEntry with the logical address and data
+	// pointers increased by NumBytes, the size decreased by NumBytes, and
+	// the same flags.
+	[[nodiscard]] FWriteLogEntry Advance(size_t NumBytes) const
 	{
-		// The address of the write.
-		std::byte* LogicalAddress = nullptr;
-		
-		// A pointer to the original data before the write occurred.
-		std::byte* Data = nullptr;
-		
-		// The size of the write in bytes.
-		size_t Size = 0;
-		
-		// Flags on the write.
-		EWriteFlags Flags = EWriteFlags::Default;
-
-		// Returns a new FWriteLogEntry with the logical address and data
-		// pointers increased by NumBytes, the size decreased by NumBytes, and
-		// the same flags.
-		[[nodiscard]] FWriteLogEntry Advance(size_t NumBytes) const
-		{
-			return {
-				.LogicalAddress = LogicalAddress + NumBytes,
-				.Data = Data + NumBytes,
-				.Size = Size - NumBytes,
-				.Flags = Flags,
-			};
-		}
+		return {
+			.LogicalAddress = LogicalAddress + NumBytes,
+			.Data = Data + NumBytes,
+			.Size = Size - NumBytes,
+			.Flags = Flags,
+		};
+	}
+};
+
+// FWriteLog holds an ordered list of write records which can be iterated
+// forwards and backwards.
+// Ensure changes to this class are kept in sync with AutoRTFM.natvis.
+class AUTORTFM_DISABLE_NO_SANITIZE FWriteLog final
+{
+	enum EAppendMode
+	{
+		Partial,
+		Whole,
 	};
 
-	// FWriteLog holds an ordered list of write records which can be iterated
-	// forwards and backwards.
-	// Ensure changes to this class are kept in sync with AutoRTFM.natvis.
-	class AUTORTFM_DISABLE_NO_SANITIZE FWriteLog final
+public:
+	template <size_t SizeBits, size_t FlagBits, size_t AddressBits>
+	struct FRecordLayout
 	{
-		enum EAppendMode
+		// Number of bits used by the FRecord to represent a write's address.
+		static constexpr size_t NumAddressBits = AddressBits;
+
+		// Number of bits used by the FRecord to represent a write's flags.
+		static constexpr size_t NumFlagsBits = FlagBits;
+
+		// Number of bits used by the FRecord to represent a write's size.
+		// Note a record size must be greater than one byte, so is stored as Size-1.
+		static constexpr size_t NumSizeBits = SizeBits;
+
+		static_assert(NumAddressBits + NumFlagsBits + NumSizeBits == 64);
+
+		static constexpr size_t AddressShift = 0;
+		static constexpr size_t AddressMask = (static_cast<size_t>(1) << NumAddressBits) - 1;
+		static constexpr size_t FlagsShift = NumAddressBits;
+		static constexpr size_t FlagsMask = (static_cast<size_t>(1) << NumFlagsBits) - 1;
+		static constexpr size_t SizeShift = NumAddressBits + NumFlagsBits;
+		static constexpr size_t SizeMask = (static_cast<size_t>(1) << NumSizeBits) - 1;
+
+		static constexpr size_t MaxAddress = AddressMask;
+		static constexpr size_t MaxSize = SizeMask + 1;
+
+		[[nodiscard]] static uint64_t Pack(std::byte* Address, EWriteFlags Flags, size_t Size)
 		{
-			Partial,
-			Whole,
-		};
-	public:
-		template<size_t SizeBits, size_t FlagBits, size_t AddressBits>
-		struct FRecordLayout
+			const size_t SizeM1 = Size - 1;
+			AUTORTFM_ASSERT_DEBUG((static_cast<size_t>(reinterpret_cast<uintptr_t>(Address)) & ~AddressMask) == 0);
+			AUTORTFM_ASSERT_DEBUG((static_cast<size_t>(Flags) & ~FlagsMask) == 0);
+			AUTORTFM_ASSERT_DEBUG((SizeM1 & ~SizeMask) == 0);
+
+			return reinterpret_cast<uint64_t>(Address) | (static_cast<uint64_t>(Flags) << FlagsShift)
+				 | (static_cast<uint64_t>(SizeM1) << SizeShift);
+		}
+
+		[[nodiscard]] static uint64_t ReplaceSize(uint64_t Packed, size_t NewSize)
 		{
-			// Number of bits used by the FRecord to represent a write's address.
-			static constexpr size_t NumAddressBits = AddressBits;
+			const size_t NewSizeM1 = NewSize - 1;
+			AUTORTFM_ASSERT_DEBUG((NewSizeM1 & ~SizeMask) == 0);
+			constexpr uint64_t Mask = ~(SizeMask << SizeShift);
+			return (Packed & Mask) | (NewSizeM1 << SizeShift);
+		}
 
-			// Number of bits used by the FRecord to represent a write's flags.
-			static constexpr size_t NumFlagsBits = FlagBits;
+		[[nodiscard]] static std::byte* UnpackAddress(uint64_t Packed)
+		{
+			return reinterpret_cast<std::byte*>((Packed >> AddressShift) & AddressMask);
+		}
 
-			// Number of bits used by the FRecord to represent a write's size.
-			// Note a record size must be greater than one byte, so is stored as Size-1.
-			static constexpr size_t NumSizeBits = SizeBits;
-			
-			static_assert(NumAddressBits + NumFlagsBits + NumSizeBits == 64);
+		[[nodiscard]] static EWriteFlags UnpackFlags(uint64_t Packed)
+		{
+			return static_cast<EWriteFlags>((Packed >> NumAddressBits) & FlagsMask);
+		}
 
-			static constexpr size_t AddressShift = 0;
-			static constexpr size_t AddressMask = (static_cast<size_t>(1) << NumAddressBits) - 1;
-			static constexpr size_t FlagsShift = NumAddressBits;
-			static constexpr size_t FlagsMask = (static_cast<size_t>(1) << NumFlagsBits) - 1;
-			static constexpr size_t SizeShift = NumAddressBits + NumFlagsBits;
-			static constexpr size_t SizeMask = (static_cast<size_t>(1) << NumSizeBits) - 1;
+		[[nodiscard]] static size_t UnpackSize(uint64_t Packed)
+		{
+			const size_t SizeM1 = reinterpret_cast<size_t>((Packed >> SizeShift) & SizeMask);
+			return SizeM1 + 1;
+		}
+	};
 
-			static constexpr size_t MaxAddress = AddressMask;
-			static constexpr size_t MaxSize = SizeMask + 1;
+	//  64           50 48                                              0
+	//  ┌─────────────┬─┬───────────────────────────────────────────────┐
+	//  │    Size     │F│                    Address                    │
+	//  └─────────────┴─┴───────────────────────────────────────────────┘
+	using FRegularRecordLayout = FRecordLayout</* SizeBits */ 14, /* FlagBits */ 2, /* AddressBits */ 48>;
 
-			[[nodiscard]] static uint64_t Pack(std::byte* Address, EWriteFlags Flags, size_t Size)
-			{
-				const size_t SizeM1 = Size - 1;
-				AUTORTFM_ASSERT_DEBUG((static_cast<size_t>(reinterpret_cast<uintptr_t>(Address)) & ~AddressMask) == 0);
-				AUTORTFM_ASSERT_DEBUG((static_cast<size_t>(Flags) & ~FlagsMask) == 0);
-				AUTORTFM_ASSERT_DEBUG((SizeM1 & ~SizeMask) == 0);
-
-				return reinterpret_cast<uint64_t>(Address)
-					| (static_cast<uint64_t>(Flags) << FlagsShift)
-					| (static_cast<uint64_t>(SizeM1) << SizeShift);
-			}
+	//  64         52   48                                              0
+	//  ┌───────────┬───┬───────────────────────────────────────────────┐
+	//  │   Size    │ F │                    Address                    │
+	//  └───────────┴───┴───────────────────────────────────────────────┘
+	using FCustomRollbackRecordLayout = FRecordLayout</* SizeBits */ 12, /* FlagBits */ 4, /* AddressBits */ 48>;
 
-			[[nodiscard]] static uint64_t ReplaceSize(uint64_t Packed, size_t NewSize)
-			{
-				const size_t NewSizeM1 = NewSize - 1;
-				AUTORTFM_ASSERT_DEBUG((NewSizeM1 & ~SizeMask) == 0);
-				constexpr uint64_t Mask = ~(SizeMask << SizeShift);
-				return (Packed & Mask) | (NewSizeM1 << SizeShift);
-			}
+private:
+	struct FRecord
+	{
+		// Returns true if the record has the EWriteFlags::CustomRollback flag set.
+		// This alters the bit packing layout of the record.
+		bool UsesCustomRollback() const
+		{
+			return (FRegularRecordLayout::UnpackFlags(Bits) & EWriteFlags::CustomRollback) == EWriteFlags::CustomRollback;
+		}
 
-			[[nodiscard]] static std::byte* UnpackAddress(uint64_t Packed)
-			{
-				return reinterpret_cast<std::byte*>((Packed >> AddressShift) & AddressMask);
-			}
+		// Packs the address, flags and size into the record using the provided bit packing layout.
+		template <typename Layout>
+		void Set(std::byte* Address, EWriteFlags Flags, size_t Size)
+		{
+			AUTORTFM_ASSERT_DEBUG(Size > 0);
+			Bits = Layout::Pack(Address, Flags, Size);
+		}
 
-			[[nodiscard]] static EWriteFlags UnpackFlags(uint64_t Packed)
-			{
-				return static_cast<EWriteFlags>((Packed >> NumAddressBits) & FlagsMask);
-			}
+		// Replaces the size field of the record using the provided bit packing layout.
+		template <typename Layout>
+		void SetSize(size_t NewSize)
+		{
+			AUTORTFM_ASSERT_DEBUG(NewSize > 0);
+			Bits = Layout::ReplaceSize(Bits, NewSize);
+		}
 
-			[[nodiscard]] static size_t UnpackSize(uint64_t Packed)
-			{
-				const size_t SizeM1 = reinterpret_cast<size_t>((Packed >> SizeShift) & SizeMask);
-				return SizeM1 + 1;
-			}
-		};
+		// Returns the unpacked address from the record.
+		std::byte* Address() const
+		{
+			return UsesCustomRollback() ? FCustomRollbackRecordLayout::UnpackAddress(Bits)
+										: FRegularRecordLayout::UnpackAddress(Bits);
+		}
 
-		//  64           50 48                                              0
-		//  ┌─────────────┬─┬───────────────────────────────────────────────┐
-		//  │    Size     │F│                    Address                    │
-		//  └─────────────┴─┴───────────────────────────────────────────────┘
-		using FRegularRecordLayout = FRecordLayout</* SizeBits */ 14, /* FlagBits */ 2, /* AddressBits */ 48>;
+		// Returns the unpacked flags from the record.
+		EWriteFlags Flags() const
+		{
+			return UsesCustomRollback() ? FCustomRollbackRecordLayout::UnpackFlags(Bits) : FRegularRecordLayout::UnpackFlags(Bits);
+		}
 
-		//  64         52   48                                              0
-		//  ┌───────────┬───┬───────────────────────────────────────────────┐
-		//  │   Size    │ F │                    Address                    │
-		//  └───────────┴───┴───────────────────────────────────────────────┘
-		using FCustomRollbackRecordLayout = FRecordLayout</* SizeBits */ 12, /* FlagBits */ 4, /* AddressBits */ 48>;
+		// Returns the unpacked size from the record.
+		size_t Size() const
+		{
+			return UsesCustomRollback() ? FCustomRollbackRecordLayout::UnpackSize(Bits) : FRegularRecordLayout::UnpackSize(Bits);
+		}
 
-	private:
-		struct FRecord
+		// Returns the unpacked size from the record, using the bit packing layout.
+		template <typename Layout>
+		size_t Size() const
 		{
-			// Returns true if the record has the EWriteFlags::CustomRollback flag set.
-			// This alters the bit packing layout of the record.
-			bool UsesCustomRollback() const
-			{
-				return (FRegularRecordLayout::UnpackFlags(Bits) & EWriteFlags::CustomRollback) == EWriteFlags::CustomRollback;
-			}
+			return Layout::UnpackSize(Bits);
+		}
 
-			// Packs the address, flags and size into the record using the provided bit packing layout.
-			template <typename Layout>
-			void Set(std::byte* Address, EWriteFlags Flags, size_t Size)
-			{
-				AUTORTFM_ASSERT_DEBUG(Size > 0);
-				Bits = Layout::Pack(Address, Flags, Size);
-			}
+		uint64_t Bits;
+	};
 
-			// Replaces the size field of the record using the provided bit packing layout.
-			template <typename Layout>
-			void SetSize(size_t NewSize)
-			{
-				AUTORTFM_ASSERT_DEBUG(NewSize > 0);
-				Bits = Layout::ReplaceSize(Bits, NewSize);
-			}
+	static_assert(sizeof(uintptr_t) == 8, "assumption: a pointer is 8 bytes");
+	static_assert(sizeof(FRecord) == 8);
 
-			// Returns the unpacked address from the record.
-			std::byte* Address() const
-			{
-				return UsesCustomRollback() ? FCustomRollbackRecordLayout::UnpackAddress(Bits) : FRegularRecordLayout::UnpackAddress(Bits);
-			}
+	// Ensure changes to this structure are kept in sync with Unreal.natvis.
+	struct FBlock final
+	{
+		// ┌────────┬────┬────┬────┬────┬────────────────┬────┬────┬────┬────┐
+		// │ FBlock │ D₀ │ D₁ │ D₂ │ D₃ │->            <-│ R₃ │ R₂ │ R₁ │ R₀ │
+		// └────────┴────┴────┴────┴────┴────────────────┴────┴────┴────┴────┘
+		//          ^                   ^                ^              ^
+		//      DataStart()          DataEnd         LastRecord    FirstRecord
+		// Where:
+		//   Dₙ = Data n, Rₙ = Record n
 
-			// Returns the unpacked flags from the record.
-			EWriteFlags Flags() const
-			{
-				return UsesCustomRollback() ? FCustomRollbackRecordLayout::UnpackFlags(Bits) : FRegularRecordLayout::UnpackFlags(Bits);
-			}
+		// Starting size of a heap-allocated block, including the FBlock struct header.
+		static constexpr size_t DefaultSize = 2048;
 
-			// Returns the unpacked size from the record.
-			size_t Size() const
-			{
-				return UsesCustomRollback() ? FCustomRollbackRecordLayout::UnpackSize(Bits) : FRegularRecordLayout::UnpackSize(Bits);
-			}
+		// Constructor
+		// TotalSize is the total size of the allocated memory for the block including
+		// the FBlock header.
+		explicit FBlock(size_t TotalSize)
+		{
+			AUTORTFM_ENSURE((TotalSize & (alignof(FRecord) - 1)) == 0);
+			std::byte* End = reinterpret_cast<std::byte*>(this) + TotalSize;
+			DataEnd = DataStart();
+			// Note: The initial empty state has LastRecord pointing one
+			// FRecord beyond the immutable FirstRecord.
+			LastRecord = reinterpret_cast<FRecord*>(End);
+			FirstRecord = LastRecord - 1;
+		}
 
-			// Returns the unpacked size from the record, using the bit packing layout.
-			template <typename Layout>
-			size_t Size() const
-			{
-				return Layout::UnpackSize(Bits);
-			}
+		// Allocate performs a heap allocation of a new block.
+		// TotalSize is the total size of the allocated memory for the block including
+		// the FBlock header.
+		static FBlock* Allocate(size_t TotalSize)
+		{
+			AUTORTFM_ASSERT(TotalSize > (sizeof(FBlock) + sizeof(FRecord)));
+			void* Memory = AutoRTFM::Allocate(TotalSize, alignof(FBlock));
+			// Disable false-positive warning C6386: Buffer overrun while writing to 'Memory'
+			CA_SUPPRESS(6386)
+			return new (Memory) FBlock(TotalSize);
+		}
 
-			uint64_t Bits;
-		};
+		// Free releases the heap-allocated memory for this block.
+		// Note: This block must have been allocated with a call to Allocate().
+		void Free()
+		{
+			AutoRTFM::Free(this);
+		}
 
-		static_assert(sizeof(uintptr_t) == 8, "assumption: a pointer is 8 bytes");
-		static_assert(sizeof(FRecord) == 8);
+		// Returns a pointer to the data for the first entry
+		std::byte* DataStart()
+		{
+			return reinterpret_cast<std::byte*>(this) + sizeof(FBlock);
+		}
 
-		// Ensure changes to this structure are kept in sync with Unreal.natvis.
-		struct FBlock final
+		// Returns a pointer to the data for the last entry
+		std::byte* LastData()
 		{
-			// ┌────────┬────┬────┬────┬────┬────────────────┬────┬────┬────┬────┐
-			// │ FBlock │ D₀ │ D₁ │ D₂ │ D₃ │->            <-│ R₃ │ R₂ │ R₁ │ R₀ │
-			// └────────┴────┴────┴────┴────┴────────────────┴────┴────┴────┴────┘
-			//          ^                   ^                ^              ^
-			//      DataStart()          DataEnd         LastRecord    FirstRecord
-			// Where:
-			//   Dₙ = Data n, Rₙ = Record n
+			return DataEnd - LastRecord->Size();
+		}
 
-			// Starting size of a heap-allocated block, including the FBlock struct header.
-			static constexpr size_t DefaultSize = 2048;
+		// Returns true if the block holds no entries.
+		bool IsEmpty() const
+		{
+			return LastRecord > FirstRecord;
+		}
 
-			// Constructor
-			// TotalSize is the total size of the allocated memory for the block including
-			// the FBlock header.
-			explicit FBlock(size_t TotalSize)
-			{
-				AUTORTFM_ENSURE((TotalSize & (alignof(FRecord) - 1)) == 0);
-				std::byte* End = reinterpret_cast<std::byte*>(this) + TotalSize;
-				DataEnd = DataStart();
-				// Note: The initial empty state has LastRecord pointing one
-				// FRecord beyond the immutable FirstRecord.
-				LastRecord = reinterpret_cast<FRecord*>(End);
-				FirstRecord = LastRecord - 1;
-			}
+		// Returns the number of unused bytes in the block between the last
+		// record and end of the data.
+		UE_AUTORTFM_FORCEINLINE size_t SpaceRemaining() const
+		{
+			return reinterpret_cast<std::byte*>(LastRecord) - DataEnd;
+		}
 
-			// Allocate performs a heap allocation of a new block.
-			// TotalSize is the total size of the allocated memory for the block including
-			// the FBlock header.
-			static FBlock* Allocate(size_t TotalSize)
+		// Appends a new record to the block, returning the number of bytes
+		// copied to the new record from the front of DataIn.
+		//
+		// If AppendMode is Partial, then AppendWithNewRecord() will append
+		// as many bytes that will fit in the remaining available space of
+		// the block, which may be fewer than NumBytes. In this mode
+		// NumBytes may be 0, in which case the function does nothing and
+		// returns 0.
+		//
+		// If AppendMode is Whole, then it is the caller's responsibility
+		// to ensure that there is enough space in the block to hold an
+		// additional NumBytes of data plus the size of an FRecord. In this
+		// mode NumBytes must be greater than zero.
+		//
+		// Layout is the bit-packing layout of a record with the given flags.
+		template <EAppendMode AppendMode, typename Layout>
+		UE_AUTORTFM_FORCEINLINE size_t AppendWithNewRecord(
+			std::byte* LogicalAddress, std::byte* DataIn, size_t NumBytes, EWriteFlags Flags)
+		{
+			if constexpr (AppendMode == Partial)
 			{
-				AUTORTFM_ASSERT(TotalSize > (sizeof(FBlock) + sizeof(FRecord)));
-				void* Memory = AutoRTFM::Allocate(TotalSize, alignof(FBlock));
-				// Disable false-positive warning C6386: Buffer overrun while writing to 'Memory'
-				CA_SUPPRESS(6386)
-				return new (Memory) FBlock(TotalSize);
+				const size_t Space = SpaceRemaining();
+				if (NumBytes == 0 || Space <= sizeof(FRecord))
+				{
+					return 0;
+				}
+				NumBytes = std::min(std::min(NumBytes, Space - sizeof(FRecord)), Layout::MaxSize);
 			}
-
-			// Free releases the heap-allocated memory for this block.
-			// Note: This block must have been allocated with a call to Allocate().
-			void Free()
+			else
 			{
-				AutoRTFM::Free(this);
+				static_assert(AppendMode == Whole);
+				AUTORTFM_ASSERT_DEBUG(NumBytes > 0);
 			}
 
-			// Returns a pointer to the data for the first entry
-			std::byte* DataStart()
-			{
-				return reinterpret_cast<std::byte*>(this) + sizeof(FBlock);
-			}
+			AUTORTFM_ASSERT_DEBUG(SpaceRemaining() >= sizeof(FRecord) + NumBytes);
 
-			// Returns a pointer to the data for the last entry
-			std::byte* LastData()
-			{
-				return DataEnd - LastRecord->Size();
-			}
+			LastRecord--;
+			LastRecord->Set<Layout>(LogicalAddress, Flags, NumBytes);
+			memcpy(DataEnd, DataIn, NumBytes);
+			DataEnd += NumBytes;
+			AUTORTFM_ASSERT_DEBUG(DataEnd <= reinterpret_cast<std::byte*>(LastRecord));
 
-			// Returns true if the block holds no entries.
-			bool IsEmpty() const
-			{
-				return LastRecord > FirstRecord;
-			}
+			return NumBytes;
+		}
 
-			// Returns the number of unused bytes in the block between the last
-			// record and end of the data.
-			UE_AUTORTFM_FORCEINLINE size_t SpaceRemaining() const
+		// Appends additional data to the tail record, returning the number
+		// of bytes copied to the end of the record from the front of DataIn.
+		//
+		// Regardless of AppendMode, if the last record cannot be appended
+		// to, either because the logical address does not start at the end
+		// of the last record's logical address or the flags are
+		// incompatible, then the function returns 0.
+		//
+		// If AppendMode is Partial, then AppendWithNewRecord() will append
+		// as many bytes that will fit in the remaining available space of
+		// the block, which may be fewer than NumBytes.
+		//
+		// If AppendMode is Whole, then AppendWithNewRecord() will only
+		// append the full NumBytes if they fit, otherwise the function will
+		// do nothing and return 0.
+		//
+		// Layout is the bit-packing layout of a record with the given flags.
+		template <EAppendMode AppendMode, typename Layout>
+		UE_AUTORTFM_FORCEINLINE size_t AppendToLastRecord(
+			std::byte* LogicalAddress, std::byte* DataIn, size_t NumBytes, EWriteFlags Flags)
+		{
+			if (IsEmpty())
 			{
-				return reinterpret_cast<std::byte*>(LastRecord) - DataEnd;
+				return 0;
 			}
 
-			// Appends a new record to the block, returning the number of bytes
-			// copied to the new record from the front of DataIn.
-			//
-			// If AppendMode is Partial, then AppendWithNewRecord() will append
-			// as many bytes that will fit in the remaining available space of
-			// the block, which may be fewer than NumBytes. In this mode
-			// NumBytes may be 0, in which case the function does nothing and
-			// returns 0.
-			//
-			// If AppendMode is Whole, then it is the caller's responsibility
-			// to ensure that there is enough space in the block to hold an
-			// additional NumBytes of data plus the size of an FRecord. In this
-			// mode NumBytes must be greater than zero.
-			//
-			// Layout is the bit-packing layout of a record with the given flags.
-			template<EAppendMode AppendMode, typename Layout>
-			UE_AUTORTFM_FORCEINLINE size_t AppendWithNewRecord(std::byte* LogicalAddress, std::byte* DataIn, size_t NumBytes, EWriteFlags Flags)
+			const size_t LastRecordSize = LastRecord->Size();
+			if (LogicalAddress != LastRecord->Address() + LastRecordSize || LastRecord->Flags() != Flags)
 			{
-				if constexpr (AppendMode == Partial)
-				{
-					const size_t Space = SpaceRemaining();
-					if (NumBytes == 0 || Space <= sizeof(FRecord))
-					{
-						return 0;
-					}
-					NumBytes = std::min(std::min(NumBytes, Space - sizeof(FRecord)), Layout::MaxSize);
-				}
-				else
-				{
-					static_assert(AppendMode == Whole);
-					AUTORTFM_ASSERT_DEBUG(NumBytes > 0);
-				}
-
-				AUTORTFM_ASSERT_DEBUG(SpaceRemaining() >= sizeof(FRecord) + NumBytes);
-
-				LastRecord--;
-				LastRecord->Set<Layout>(LogicalAddress, Flags, NumBytes);
-				memcpy(DataEnd, DataIn, NumBytes);
-				DataEnd += NumBytes;
-				AUTORTFM_ASSERT_DEBUG(DataEnd <= reinterpret_cast<std::byte*>(LastRecord));
-
-				return NumBytes;
+				return 0;
 			}
 
-			// Appends additional data to the tail record, returning the number
-			// of bytes copied to the end of the record from the front of DataIn.
-			//
-			// Regardless of AppendMode, if the last record cannot be appended
-			// to, either because the logical address does not start at the end
-			// of the last record's logical address or the flags are
-			// incompatible, then the function returns 0.
-			//
-			// If AppendMode is Partial, then AppendWithNewRecord() will append
-			// as many bytes that will fit in the remaining available space of
-			// the block, which may be fewer than NumBytes.
-			//
-			// If AppendMode is Whole, then AppendWithNewRecord() will only
-			// append the full NumBytes if they fit, otherwise the function will
-			// do nothing and return 0.
-			//
-			// Layout is the bit-packing layout of a record with the given flags.
-			template<EAppendMode AppendMode, typename Layout>
-			UE_AUTORTFM_FORCEINLINE size_t AppendToLastRecord(std::byte* LogicalAddress, std::byte* DataIn, size_t NumBytes, EWriteFlags Flags)
+			const size_t RepresentableSizeRemaining = Layout::MaxSize - LastRecordSize;
+			const size_t Available = std::min(RepresentableSizeRemaining, SpaceRemaining());
+
+			if constexpr (AppendMode == Whole)
 			{
-				if (IsEmpty())
-				{
-					return 0;
-				}
-				
-				const size_t LastRecordSize = LastRecord->Size();
-				if (LogicalAddress != LastRecord->Address() + LastRecordSize ||
-					LastRecord->Flags() != Flags)
+				if (NumBytes > Available)
 				{
 					return 0;
 				}
+			}
+			else
+			{
+				static_assert(AppendMode == Partial);
+				NumBytes = std::min(NumBytes, Available);
+			}
 
-				const size_t RepresentableSizeRemaining = Layout::MaxSize - LastRecordSize;
-				const size_t Available = std::min(RepresentableSizeRemaining, SpaceRemaining());
+			LastRecord->SetSize<Layout>(LastRecordSize + NumBytes);
+			memcpy(DataEnd, DataIn, NumBytes);
+			DataEnd += NumBytes;
+			AUTORTFM_ASSERT_DEBUG(DataEnd <= reinterpret_cast<std::byte*>(LastRecord));
+			return NumBytes;
+		}
 
-				if constexpr (AppendMode == Whole)
-				{
-					if (NumBytes > Available)
-					{
-						return 0;
-					}
-				}
-				else
-				{
-					static_assert(AppendMode == Partial);
-					NumBytes = std::min(NumBytes, Available);
-				}
+		// The next block in the linked list.
+		FBlock* NextBlock = nullptr;
+		// The previous block in the linked list.
+		FBlock* PrevBlock = nullptr;
+		// The pointer to the first entry's record
+		FRecord* FirstRecord = nullptr;
+		// The pointer to the last entry's record
+		FRecord* LastRecord = nullptr;
+		// One byte beyond the end of the last entry's data
+		std::byte* DataEnd = nullptr;
 
-				LastRecord->SetSize<Layout>(LastRecordSize + NumBytes);
-				memcpy(DataEnd, DataIn, NumBytes);
-				DataEnd += NumBytes;
-				AUTORTFM_ASSERT_DEBUG(DataEnd <= reinterpret_cast<std::byte*>(LastRecord));
-				return NumBytes;
-			}
+	private:
+		~FBlock() = delete;
+	};
 
-			// The next block in the linked list.
-			FBlock* NextBlock = nullptr;
-			// The previous block in the linked list.
-			FBlock* PrevBlock = nullptr;
-			// The pointer to the first entry's record
-			FRecord* FirstRecord = nullptr;
-			// The pointer to the last entry's record
-			FRecord* LastRecord = nullptr;
-			// One byte beyond the end of the last entry's data
-			std::byte* DataEnd = nullptr;
-		private:
-			~FBlock() = delete;
-		};
+public:
+	// Constructor
+	FWriteLog()
+	{
+		new (HeadBlockMemory) FBlock(HeadBlockSize);
+	}
 
-	public:
-		// Constructor
-		FWriteLog()
+	// Destructor
+	~FWriteLog()
+	{
+		Reset();
+	}
+
+	// Adds the write log entry to the log.
+	// The log will make a copy of the FWriteLogEntry's data.
+	void Push(FWriteLogEntry Entry)
+	{
+		if ((Entry.Flags & EWriteFlags::CustomRollback) == EWriteFlags::CustomRollback)
 		{
-			new(HeadBlockMemory) FBlock(HeadBlockSize);
+			Push<FCustomRollbackRecordLayout>(Entry);
 		}
-
-		// Destructor
-		~FWriteLog()
+		else
 		{
-			Reset();
+			Push<FRegularRecordLayout>(Entry);
 		}
+	}
 
-		// Adds the write log entry to the log.
-		// The log will make a copy of the FWriteLogEntry's data.
-		void Push(FWriteLogEntry Entry)
+	// Adds the write log entry to the log; assumes a payload small enough that splitting is not beneficial.
+	// If you have large sizes, you should use `Push` instead so that splitting can occur.
+	template <unsigned int Size>
+	void PushSmall(std::byte* LogicalAddress, std::byte* Data, EWriteFlags Flags = EWriteFlags::Default)
+	{
+		if ((Flags & EWriteFlags::CustomRollback) == EWriteFlags::CustomRollback)
 		{
-			if ((Entry.Flags & EWriteFlags::CustomRollback) == EWriteFlags::CustomRollback)
-			{
-				Push<FCustomRollbackRecordLayout>(Entry);
-			}
-			else
-			{
-				Push<FRegularRecordLayout>(Entry);
-			}
+			PushSmall<Size, FCustomRollbackRecordLayout>(LogicalAddress, Data, Flags);
 		}
-
-		// Adds the write log entry to the log; assumes a payload small enough that splitting is not beneficial.
-		// If you have large sizes, you should use `Push` instead so that splitting can occur.
-		template <unsigned int Size>
-		void PushSmall(std::byte* LogicalAddress, std::byte* Data, EWriteFlags Flags = EWriteFlags::Default)
+		else
 		{
-			if ((Flags & EWriteFlags::CustomRollback) == EWriteFlags::CustomRollback)
-			{
-				PushSmall<Size, FCustomRollbackRecordLayout>(LogicalAddress, Data, Flags);
-			}
-			else
-			{
-				PushSmall<Size, FRegularRecordLayout>(LogicalAddress, Data, Flags);
-			}
+			PushSmall<Size, FRegularRecordLayout>(LogicalAddress, Data, Flags);
 		}
+	}
 
-		// Iterator for enumerating the writes of the log.
-		template<bool IS_FORWARD>
-		struct TIterator final
-		{
-			TIterator() = default;
+	// Iterator for enumerating the writes of the log.
+	template <bool IS_FORWARD>
+	struct TIterator final
+	{
+		TIterator() = default;
 
-			TIterator(FBlock* StartBlock) : Block(StartBlock)
+		TIterator(FBlock* StartBlock) : Block(StartBlock)
+		{
+			if (UNLIKELY(Block->IsEmpty()))
 			{
-				if (UNLIKELY(Block->IsEmpty()))
+				if (UNLIKELY(!AdvanceBlock()))
 				{
-					if (UNLIKELY(!AdvanceBlock()))
-					{
-						// The write log is entirely empty.
-						return;
-					}
+					// The write log is entirely empty.
+					return;
 				}
-
-				Data = IS_FORWARD ? Block->DataStart() : Block->LastData();
-				Record = IS_FORWARD ? Block->FirstRecord : Block->LastRecord;
 			}
 
-			// Returns the entry at the current iterator's position.
-			FWriteLogEntry operator*() const
-			{
-				return FWriteLogEntry
-				{
-					.LogicalAddress = reinterpret_cast<std::byte*>(Record->Address()),
-					.Data = Data,
-					.Size = Record->Size(),
-					.Flags = Record->Flags(),
-				};
-			}
+			Data = IS_FORWARD ? Block->DataStart() : Block->LastData();
+			Record = IS_FORWARD ? Block->FirstRecord : Block->LastRecord;
+		}
 
-			// Progresses the iterator to the next entry
-			void operator++()
+		// Returns the entry at the current iterator's position.
+		FWriteLogEntry operator*() const
+		{
+			return FWriteLogEntry{
+				.LogicalAddress = reinterpret_cast<std::byte*>(Record->Address()),
+				.Data = Data,
+				.Size = Record->Size(),
+				.Flags = Record->Flags(),
+			};
+		}
+
+		// Progresses the iterator to the next entry
+		void operator++()
+		{
+			if constexpr (IS_FORWARD)
 			{
-				if constexpr (IS_FORWARD)
+				if (Record == Block->LastRecord)
 				{
-					if (Record == Block->LastRecord)
-					{
-						if (LIKELY(AdvanceBlock()))
-						{
-							Data = Block->DataStart();
-							Record = Block->FirstRecord;
-						}
-					}
-					else
+					if (LIKELY(AdvanceBlock()))
 					{
-						Data += Record->Size();
-						Record--;
+						Data = Block->DataStart();
+						Record = Block->FirstRecord;
 					}
 				}
 				else
 				{
-					if (Record == Block->FirstRecord)
-					{
-						if (LIKELY(AdvanceBlock()))
-						{
-							Data = Block->LastData();
-							Record = Block->LastRecord;
-						}
-					}
-					else
-					{
-						Record++;
-						Data -= Record->Size();
-					}
+					Data += Record->Size();
+					Record--;
 				}
 			}
-
-			// Inequality operator
-			bool operator!=(const TIterator& Other) const
-			{
-				return (Other.Block != Block) || (Other.Record != Record);
-			}
-
-		private:
-			// Resets the iterator (compares equal to the write log's end())
-			UE_AUTORTFM_FORCEINLINE void Reset()
-			{
-				Block = nullptr;
-				Data = nullptr;
-				Record = nullptr;
-			}
-
-			// Moves from this block to the next (if IS_FORWARD) or previous (if not IS_FORWARD),
-			// skipping any empty blocks. Returns true on success. If no more blocks exist, resets 
-			// the iterator and returns false.
-			UE_AUTORTFM_FORCEINLINE bool AdvanceBlock()
+			else
 			{
-				do
+				if (Record == Block->FirstRecord)
 				{
-					Block = IS_FORWARD ? Block->NextBlock : Block->PrevBlock;
+					if (LIKELY(AdvanceBlock()))
+					{
+						Data = Block->LastData();
+						Record = Block->LastRecord;
+					}
 				}
-				while (Block && Block->IsEmpty());
-
-				if (!Block)
+				else
 				{
-					Reset();
-					return false;
+					Record++;
+					Data -= Record->Size();
 				}
-
-				return true;
 			}
+		}
 
-			FBlock* Block = nullptr;
-			std::byte* Data = nullptr;
-			FRecord* Record = nullptr;
-		};
-
-		using Iterator = TIterator</* IS_FORWARD */ true>;
-		using ReverseIterator = TIterator</* IS_FORWARD */ false>;
-
-		Iterator begin() const
+		// Inequality operator
+		bool operator!=(const TIterator& Other) const
 		{
-			return (NumEntries > 0) ? Iterator(HeadBlock) : Iterator{};
+			return (Other.Block != Block) || (Other.Record != Record);
 		}
-		ReverseIterator rbegin() const
+
+	private:
+		// Resets the iterator (compares equal to the write log's end())
+		UE_AUTORTFM_FORCEINLINE void Reset()
 		{
-			return (NumEntries > 0) ? ReverseIterator(TailBlock) : ReverseIterator{};
+			Block = nullptr;
+			Data = nullptr;
+			Record = nullptr;
 		}
-		Iterator end() const { return Iterator{}; }
-		ReverseIterator rend() const { return ReverseIterator{}; }
 
-		// Resets the write log to its initial state, freeing any allocated memory.
-		void Reset()
+		// Moves from this block to the next (if IS_FORWARD) or previous (if not IS_FORWARD),
+		// skipping any empty blocks. Returns true on success. If no more blocks exist, resets
+		// the iterator and returns false.
+		UE_AUTORTFM_FORCEINLINE bool AdvanceBlock()
 		{
-			// Skip HeadBlock, which is held as part of this structure.
-			FBlock* Block = HeadBlock->NextBlock;
-			while (nullptr != Block)
+			do
+			{
+				Block = IS_FORWARD ? Block->NextBlock : Block->PrevBlock;
+			} while (Block && Block->IsEmpty());
+
+			if (!Block)
 			{
-				FBlock* const Next = Block->NextBlock;
-				Block->Free();
-				Block = Next;
+				Reset();
+				return false;
 			}
-			new (HeadBlockMemory) FBlock(HeadBlockSize);
-			HeadBlock = reinterpret_cast<FBlock*>(HeadBlockMemory);
-			TailBlock = reinterpret_cast<FBlock*>(HeadBlockMemory);
-			NumEntries = 0;
-			TotalSizeBytes = 0;
-			NextBlockSize = FBlock::DefaultSize;
+
+			return true;
 		}
 
-		// Returns true if the log holds no entries.
-		UE_AUTORTFM_FORCEINLINE bool IsEmpty() const { return 0 == NumEntries; }
+		FBlock* Block = nullptr;
+		std::byte* Data = nullptr;
+		FRecord* Record = nullptr;
+	};
 
-		// Return the number of entries in the log.
-		UE_AUTORTFM_FORCEINLINE size_t Num() const { return NumEntries; }
+	using Iterator = TIterator</* IS_FORWARD */ true>;
+	using ReverseIterator = TIterator</* IS_FORWARD */ false>;
 
-		// Return the total size in bytes for all entries in the log.
-		UE_AUTORTFM_FORCEINLINE size_t TotalSize() const { return TotalSizeBytes; }
+	Iterator begin() const
+	{
+		return (NumEntries > 0) ? Iterator(HeadBlock) : Iterator{};
+	}
+	ReverseIterator rbegin() const
+	{
+		return (NumEntries > 0) ? ReverseIterator(TailBlock) : ReverseIterator{};
+	}
+	Iterator end() const
+	{
+		return Iterator{};
+	}
+	ReverseIterator rend() const
+	{
+		return ReverseIterator{};
+	}
 
-	private:
-		template<typename Layout>
-		void Push(FWriteLogEntry Entry)
+	// Resets the write log to its initial state, freeing any allocated memory.
+	void Reset()
+	{
+		// Skip HeadBlock, which is held as part of this structure.
+		FBlock* Block = HeadBlock->NextBlock;
+		while (nullptr != Block)
 		{
-			TotalSizeBytes += Entry.Size;
-
-			// Attempt to append to the end of the last record.
-			if (const size_t Appended = TailBlock->AppendToLastRecord<Partial, Layout>(Entry.LogicalAddress, Entry.Data, Entry.Size, Entry.Flags))
-			{
-				Entry = Entry.Advance(Appended);
-			}
+			FBlock* const Next = Block->NextBlock;
+			Block->Free();
+			Block = Next;
+		}
+		new (HeadBlockMemory) FBlock(HeadBlockSize);
+		HeadBlock = reinterpret_cast<FBlock*>(HeadBlockMemory);
+		TailBlock = reinterpret_cast<FBlock*>(HeadBlockMemory);
+		NumEntries = 0;
+		TotalSizeBytes = 0;
+		NextBlockSize = FBlock::DefaultSize;
+	}
+
+	// Returns true if the log holds no entries.
+	UE_AUTORTFM_FORCEINLINE bool IsEmpty() const
+	{
+		return 0 == NumEntries;
+	}
 
-			// Append to the current tail block with new records until we've
-			// either filled the block, or run out of data to append.
-			while (const size_t Appended = TailBlock->AppendWithNewRecord<Partial, Layout>(Entry.LogicalAddress, Entry.Data, Entry.Size, Entry.Flags))
-			{
-				NumEntries++;
-				Entry = Entry.Advance(Appended);
-			}
+	// Return the number of entries in the log.
+	UE_AUTORTFM_FORCEINLINE size_t Num() const
+	{
+		return NumEntries;
+	}
 
-			if (AUTORTFM_LIKELY(Entry.Size == 0))
-			{
-				return;
-			}
+	// Return the total size in bytes for all entries in the log.
+	UE_AUTORTFM_FORCEINLINE size_t TotalSize() const
+	{
+		return TotalSizeBytes;
+	}
 
-			// TailBlock is full, and there's more data to append.
-			// We need to allocate a new block.
-
-			// The maximum number of bytes that can be stored in a single record.
-			const size_t RecordMaxSize = Layout::MaxSize;
-			// Calculate how many records we need to make.
-			const size_t NumRecords = (Entry.Size + RecordMaxSize - 1) / RecordMaxSize;
-			// Calculate the exact required size of the block.
-			const size_t RequiredSize =
-				sizeof(FBlock) +               // FBlock header
-				NumRecords * sizeof(FRecord) + // Bytes needed for the records
-				Entry.Size;                    // Bytes needed for the data
-
-			// Create a new empty tail block, large enough to hold the remainder
-			// of this push in its entirety, and never smaller than the upcoming
-			// block size.
-			AllocateNewBlock(std::max(RequiredSize, NextBlockSize));
-
-			// Append (NumRecords-1) full-records.
-			for (size_t I = 1; I < NumRecords; I++)
-			{
-				TailBlock->AppendWithNewRecord<Whole, Layout>(Entry.LogicalAddress, Entry.Data, RecordMaxSize, Entry.Flags);
-				NumEntries++;
-				Entry = Entry.Advance(RecordMaxSize);
-			}
+private:
+	template <typename Layout>
+	void Push(FWriteLogEntry Entry)
+	{
+		TotalSizeBytes += Entry.Size;
 
-			// Append last partial record.
-			if (Entry.Size)
-			{
-				TailBlock->AppendWithNewRecord<Whole, Layout>(Entry.LogicalAddress, Entry.Data, Entry.Size, Entry.Flags);
-				NumEntries++;
-			}
+		// Attempt to append to the end of the last record.
+		if (const size_t Appended =
+				TailBlock->AppendToLastRecord<Partial, Layout>(Entry.LogicalAddress, Entry.Data, Entry.Size, Entry.Flags))
+		{
+			Entry = Entry.Advance(Appended);
 		}
 
-		template <size_t Size, typename Layout>
-		void PushSmall(std::byte* LogicalAddress, std::byte* Data, EWriteFlags Flags)
+		// Append to the current tail block with new records until we've
+		// either filled the block, or run out of data to append.
+		while (const size_t Appended =
+				   TailBlock->AppendWithNewRecord<Partial, Layout>(Entry.LogicalAddress, Entry.Data, Entry.Size, Entry.Flags))
 		{
-			static_assert(Size <= Layout::MaxSize);
+			NumEntries++;
+			Entry = Entry.Advance(Appended);
+		}
 
-			TotalSizeBytes += Size;
+		if (AUTORTFM_LIKELY(Entry.Size == 0))
+		{
+			return;
+		}
 
-			if (TailBlock->AppendToLastRecord<Whole, Layout>(LogicalAddress, Data, Size, Flags))
-			{
-				return;
-			}
-			
-			if (TailBlock->SpaceRemaining() < sizeof(FRecord) + Size)
-			{
-				AllocateNewBlock(NextBlockSize);
-			}
-			
-			TailBlock->AppendWithNewRecord<Whole, Layout>(LogicalAddress, Data, Size, Flags);
+		// TailBlock is full, and there's more data to append.
+		// We need to allocate a new block.
+
+		// The maximum number of bytes that can be stored in a single record.
+		const size_t RecordMaxSize = Layout::MaxSize;
+		// Calculate how many records we need to make.
+		const size_t NumRecords = (Entry.Size + RecordMaxSize - 1) / RecordMaxSize;
+		// Calculate the exact required size of the block.
+		const size_t RequiredSize = sizeof(FBlock) +                // FBlock header
+									NumRecords * sizeof(FRecord) +  // Bytes needed for the records
+									Entry.Size;                     // Bytes needed for the data
+
+		// Create a new empty tail block, large enough to hold the remainder
+		// of this push in its entirety, and never smaller than the upcoming
+		// block size.
+		AllocateNewBlock(std::max(RequiredSize, NextBlockSize));
+
+		// Append (NumRecords-1) full-records.
+		for (size_t I = 1; I < NumRecords; I++)
+		{
+			TailBlock->AppendWithNewRecord<Whole, Layout>(Entry.LogicalAddress, Entry.Data, RecordMaxSize, Entry.Flags);
 			NumEntries++;
+			Entry = Entry.Advance(RecordMaxSize);
 		}
 
-		UE_AUTORTFM_FORCEINLINE void AllocateNewBlock(size_t Size)
+		// Append last partial record.
+		if (Entry.Size)
 		{
-			FBlock* NewBlock = FBlock::Allocate(AlignUp(Size, sizeof(FRecord)));
-			NewBlock->PrevBlock = TailBlock;
-			TailBlock->NextBlock = NewBlock;
-			TailBlock = NewBlock;
-
-			// Increase block sizes by 50% each time, capped at MaxBlockSize.
-			NextBlockSize = std::min<size_t>((NextBlockSize * 3 / 2), MaxBlockSize);
+			TailBlock->AppendWithNewRecord<Whole, Layout>(Entry.LogicalAddress, Entry.Data, Entry.Size, Entry.Flags);
+			NumEntries++;
 		}
+	}
 
-		template <size_t SIZE>
-		static constexpr bool IsAlignedForFRecord = (SIZE & (alignof(FRecord) - 1)) == 0;
+	template <size_t Size, typename Layout>
+	void PushSmall(std::byte* LogicalAddress, std::byte* Data, EWriteFlags Flags)
+	{
+		static_assert(Size <= Layout::MaxSize);
 
-		// The size of the inline block, which is declared as a byte array (`HeadBlockMemory`)
-		// inside the write log.
-		static constexpr size_t HeadBlockSize = 256;
+		TotalSizeBytes += Size;
 
-		// The upper bound on heap-allocated block size.
-		static constexpr size_t MaxBlockDataSize = 1024 * 1024 * 16; // Without header, record or alignment.
-		static constexpr size_t MaxBlockSize = AlignUp(sizeof(FBlock) + sizeof(FRecord) + MaxBlockDataSize, alignof(FRecord));
+		if (TailBlock->AppendToLastRecord<Whole, Layout>(LogicalAddress, Data, Size, Flags))
+		{
+			return;
+		}
 
-		static_assert(IsAlignedForFRecord<HeadBlockSize>);
-		static_assert(IsAlignedForFRecord<FBlock::DefaultSize>);
-		static_assert(IsAlignedForFRecord<MaxBlockSize>);
+		if (TailBlock->SpaceRemaining() < sizeof(FRecord) + Size)
+		{
+			AllocateNewBlock(NextBlockSize);
+		}
 
-		FBlock* HeadBlock = reinterpret_cast<FBlock*>(HeadBlockMemory);
-		FBlock* TailBlock = reinterpret_cast<FBlock*>(HeadBlockMemory);
-		size_t NumEntries = 0;
-		size_t TotalSizeBytes = 0;
-		size_t NextBlockSize = FBlock::DefaultSize;
-		alignas(alignof(FBlock)) std::byte HeadBlockMemory[HeadBlockSize];
-	};
-}
+		TailBlock->AppendWithNewRecord<Whole, Layout>(LogicalAddress, Data, Size, Flags);
+		NumEntries++;
+	}
 
-#endif // (defined(__AUTORTFM) && __AUTORTFM)
+	UE_AUTORTFM_FORCEINLINE void AllocateNewBlock(size_t Size)
+	{
+		FBlock* NewBlock = FBlock::Allocate(AlignUp(Size, sizeof(FRecord)));
+		NewBlock->PrevBlock = TailBlock;
+		TailBlock->NextBlock = NewBlock;
+		TailBlock = NewBlock;
+
+		// Increase block sizes by 50% each time, capped at MaxBlockSize.
+		NextBlockSize = std::min<size_t>((NextBlockSize * 3 / 2), MaxBlockSize);
+	}
+
+	template <size_t SIZE>
+	static constexpr bool IsAlignedForFRecord = (SIZE & (alignof(FRecord) - 1)) == 0;
+
+	// The size of the inline block, which is declared as a byte array (`HeadBlockMemory`)
+	// inside the write log.
+	static constexpr size_t HeadBlockSize = 256;
+
+	// The upper bound on heap-allocated block size.
+	static constexpr size_t MaxBlockDataSize = 1024 * 1024 * 16;  // Without header, record or alignment.
+	static constexpr size_t MaxBlockSize = AlignUp(sizeof(FBlock) + sizeof(FRecord) + MaxBlockDataSize, alignof(FRecord));
+
+	static_assert(IsAlignedForFRecord<HeadBlockSize>);
+	static_assert(IsAlignedForFRecord<FBlock::DefaultSize>);
+	static_assert(IsAlignedForFRecord<MaxBlockSize>);
+
+	FBlock* HeadBlock = reinterpret_cast<FBlock*>(HeadBlockMemory);
+	FBlock* TailBlock = reinterpret_cast<FBlock*>(HeadBlockMemory);
+	size_t NumEntries = 0;
+	size_t TotalSizeBytes = 0;
+	size_t NextBlockSize = FBlock::DefaultSize;
+	alignas(alignof(FBlock)) std::byte HeadBlockMemory[HeadBlockSize];
+};
+}  // namespace AutoRTFM
+
+#endif  // (defined(__AUTORTFM) && __AUTORTFM)
diff --git a/Engine/Source/Runtime/AutoRTFM/Public/AutoRTFM.h b/Engine/Source/Runtime/AutoRTFM/Public/AutoRTFM.h
--- a/Engine/Source/Runtime/AutoRTFM/Public/AutoRTFM.h
+++ b/Engine/Source/Runtime/AutoRTFM/Public/AutoRTFM.h
@@ -10,9 +10,9 @@
 
 // HEADER_UNIT_SKIP - unused warnings
 
+#include "AutoRTFM/CAPI.h"
 #include "AutoRTFM/Constants.h"
 #include "AutoRTFM/Defines.h"
-#include "AutoRTFM/CAPI.h"
 #include "AutoRTFM/ReturnFromOpenMode.h"
 
 #ifdef __cplusplus
@@ -21,9 +21,9 @@
 
 #include <cstdarg>
 #include <new>
+#include <stdlib.h>
 #include <type_traits>
 #include <utility>
-#include <stdlib.h>
 
 namespace AutoRTFM
 {
@@ -44,12 +44,12 @@ enum class EWriteFlags
 };
 
 // Bitwise operator | for EWriteFlags
-inline constexpr EWriteFlags operator | (EWriteFlags LHS, EWriteFlags RHS)
+inline constexpr EWriteFlags operator|(EWriteFlags LHS, EWriteFlags RHS)
 {
 	return static_cast<EWriteFlags>(static_cast<int>(LHS) | static_cast<int>(RHS));
 }
 // Bitwise operator & for EWriteFlags
-inline constexpr EWriteFlags operator & (EWriteFlags LHS, EWriteFlags RHS)
+inline constexpr EWriteFlags operator&(EWriteFlags LHS, EWriteFlags RHS)
 {
 	return static_cast<EWriteFlags>(static_cast<int>(LHS) & static_cast<int>(RHS));
 }
@@ -57,7 +57,7 @@ inline constexpr EWriteFlags operator & (EWriteFlags LHS, EWriteFlags RHS)
 // Evaluates to the AutoRTFM mode of the function overload which would be called
 // when calling FuncType with the arguments of type ArgTypes.
 // Warning: This is an experimental API and may be removed in the future.
-template <typename FuncType, typename ... ArgTypes>
+template <typename FuncType, typename... ArgTypes>
 static constexpr autortfm_mode CallMode = AUTORTFM_MODE_OF_CALL(std::declval<FuncType>()(std::declval<ArgTypes>()...));
 
 // Evaluates to the AutoRTFM mode of the destructor function of the type Type.
@@ -68,7 +68,7 @@ static constexpr autortfm_mode DestructorMode = AUTORTFM_MODE_OF_CALL(std::declv
 // Evaluates to the AutoRTFM mode of the constructor function of the type Type
 // when called with the arguments of type ArgTypes.
 // Warning: This is an experimental API and may be removed in the future.
-template <typename Type, typename ... ArgTypes>
+template <typename Type, typename... ArgTypes>
 static constexpr autortfm_mode ConstructorMode = AUTORTFM_MODE_OF_CALL(::new Type(std::declval<ArgTypes>()...));
 
 // The transaction result provides information on how a transaction completed. This is either Committed,
@@ -76,14 +76,14 @@ static constexpr autortfm_mode ConstructorMode = AUTORTFM_MODE_OF_CALL(::new Typ
 enum class ETransactionResult
 {
 	// The transaction aborted because of an explicit call to AbortTransaction or RollbackTransaction.
-    AbortedByRequest = autortfm_aborted_by_request,
+	AbortedByRequest = autortfm_aborted_by_request,
 
 	// The transaction aborted because of unhandled constructs in the code (atomics, unhandled function calls, etc).
-    AbortedByLanguage = autortfm_aborted_by_language,
+	AbortedByLanguage = autortfm_aborted_by_language,
 
 	// The transaction committed successfully. For a nested transaction this does not mean that the transaction effects
 	// cannot be undone later if the parent transaction is aborted for any reason.
-    Committed = autortfm_committed,
+	Committed = autortfm_committed,
 
 	// The transaction aborted because a new transaction nest was attempted (via OnCommit or OnComplete) while the
 	// current transaction was being committed.
@@ -159,106 +159,117 @@ using TransactionID = autortfm_transaction_id;
 // AutoRTFM sanitizer API (requires the -fsanitize=autortfm compiler flag)
 namespace Sanitizer
 {
-	// An enumerator of AutoRTFM sanitizer levels.
-	// The AutoRTFM sanitizer is used to detect modification by open-code to memory that was written by
-	// a transaction. In this situation, aborting the transaction can corrupt memory as the undo will
-	// overwrite the writes made in the open-code.
-	enum class EMode : uint8_t
-	{
-		// AutoRTFM sanitizer is disabled.
-		Disabled = autortfm_sanitizer_disabled,
+// An enumerator of AutoRTFM sanitizer levels.
+// The AutoRTFM sanitizer is used to detect modification by open-code to memory that was written by
+// a transaction. In this situation, aborting the transaction can corrupt memory as the undo will
+// overwrite the writes made in the open-code.
+enum class EMode : uint8_t
+{
+	// AutoRTFM sanitizer is disabled.
+	Disabled = autortfm_sanitizer_disabled,
 
-		// AutoRTFM sanitizer is enabled and failures are treated as warnings.
-		Warn = autortfm_sanitizer_warn,
+	// AutoRTFM sanitizer is enabled and failures are treated as warnings.
+	Warn = autortfm_sanitizer_warn,
 
-		// AutoRTFM sanitizer is enabled and failures are treated as errors.
-		Error = autortfm_sanitizer_error,
-	};
+	// AutoRTFM sanitizer is enabled and failures are treated as errors.
+	Error = autortfm_sanitizer_error,
+};
 
-	// Returns the global AutoRTFM sanitizer mode currently in use.
-	UE_AUTORTFM_API EMode GetMode();
+// Returns the global AutoRTFM sanitizer mode currently in use.
+UE_AUTORTFM_API EMode GetMode();
 
-	// Sets the global AutoRTFM sanitizer mode.
-	UE_AUTORTFM_API void SetMode(EMode Mode);
+// Sets the global AutoRTFM sanitizer mode.
+UE_AUTORTFM_API void SetMode(EMode Mode);
 
-	// Temporarily disables the AutoRTFM sanitizer for the current transaction
-	// and any child transactions.
-	// Must be paired with a call to EndDisable().
-	// Calls to BeginDisable() are internally counted and may be nested within
-	// other BeginDisable() / EndDisable() calls.
-	// This call is a no-op when called outside of a transaction.
-	UE_AUTORTFM_API AUTORTFM_OPEN void BeginDisable();
+// Temporarily disables the AutoRTFM sanitizer for the current transaction
+// and any child transactions.
+// Must be paired with a call to EndDisable().
+// Calls to BeginDisable() are internally counted and may be nested within
+// other BeginDisable() / EndDisable() calls.
+// This call is a no-op when called outside of a transaction.
+UE_AUTORTFM_API AUTORTFM_OPEN void BeginDisable();
 
-	// Ends temporarily disabling the AutoRTFM sanitizer.
-	// See BeginDisable().
-	UE_AUTORTFM_API AUTORTFM_OPEN void EndDisable();
+// Ends temporarily disabling the AutoRTFM sanitizer.
+// See BeginDisable().
+UE_AUTORTFM_API AUTORTFM_OPEN void EndDisable();
 
-	// Returns true if the AutoRTFM sanitizer records the callstack of closed writes.
-	UE_AUTORTFM_API bool RecordClosedWriteCallstacks();
+// Returns true if the AutoRTFM sanitizer records the callstack of closed writes.
+UE_AUTORTFM_API bool RecordClosedWriteCallstacks();
 
-	// Sets whether the AutoRTFM sanitizer records the callstack of closed writes.
-	UE_AUTORTFM_API void SetRecordClosedWriteCallstacks(bool bEnable);
+// Sets whether the AutoRTFM sanitizer records the callstack of closed writes.
+UE_AUTORTFM_API void SetRecordClosedWriteCallstacks(bool bEnable);
 
-	// A helper class that temporarily disables the AutoRTFM sanitizer for
-	// the lifetime of the object.
-	class FDisableScope
+// A helper class that temporarily disables the AutoRTFM sanitizer for
+// the lifetime of the object.
+class FDisableScope
+{
+public:
+	UE_AUTORTFM_FORCEINLINE FDisableScope()
 	{
-	public:
-		UE_AUTORTFM_FORCEINLINE FDisableScope() { BeginDisable(); }
-		UE_AUTORTFM_FORCEINLINE ~FDisableScope() { EndDisable(); }
-	private:
-		FDisableScope(const FDisableScope&) = delete;
-		FDisableScope& operator = (const FDisableScope&) = delete;
-	};
-}
-#endif // AUTORTFM_SANITIZER
+		BeginDisable();
+	}
+	UE_AUTORTFM_FORCEINLINE ~FDisableScope()
+	{
+		EndDisable();
+	}
+
+private:
+	FDisableScope(const FDisableScope&) = delete;
+	FDisableScope& operator=(const FDisableScope&) = delete;
+};
+}  // namespace Sanitizer
+#endif  // AUTORTFM_SANITIZER
 
 #if AUTORTFM_SANITIZER
-	#define AUTORTFM_SANITIZER_DISABLE_SCOPE() \
-		::AutoRTFM::Sanitizer::FDisableScope UE_AUTORTFM_CONCAT(AutoRTFMSanitizerModeScope, __COUNTER__)
+#define AUTORTFM_SANITIZER_DISABLE_SCOPE() \
+	::AutoRTFM::Sanitizer::FDisableScope UE_AUTORTFM_CONCAT(AutoRTFMSanitizerModeScope, __COUNTER__)
 #else
-	#define AUTORTFM_SANITIZER_DISABLE_SCOPE() static_assert(true) /* require semicolon */
+#define AUTORTFM_SANITIZER_DISABLE_SCOPE() static_assert(true) /* require semicolon */
 #endif
 
 #if UE_AUTORTFM
 namespace ForTheRuntime
 {
-	UE_AUTORTFM_API void CascadingAbortTransactionInternal();
-	UE_AUTORTFM_API void CascadingRetryTransactionInternal();
-	UE_AUTORTFM_API void OnCommitInternal(TTask<void()>&& Work);
-	UE_AUTORTFM_API void OnPreAbortInternal(TTask<void()>&& Work);
-	UE_AUTORTFM_API void OnAbortInternal(TTask<void()>&& Work);
-	UE_AUTORTFM_API void OnCompleteInternal(TTask<void()>&& Work);
-	UE_AUTORTFM_API void PushOnCommitHandlerInternal(const void* Key, TTask<void()>&& Work);
-	UE_AUTORTFM_API void PopOnCommitHandlerInternal(const void* Key);
-	UE_AUTORTFM_API void PopAllOnCommitHandlersInternal(const void* Key);
-	UE_AUTORTFM_API void PushOnAbortHandlerInternal(const void* Key, TTask<void()>&& Work);
-	UE_AUTORTFM_API void PopOnAbortHandlerInternal(const void* Key);
-	UE_AUTORTFM_API void PopAllOnAbortHandlersInternal(const void* Key);
-	UE_AUTORTFM_API void RegisterOnCommitFromTheOpen(TTask<void()>&& Work);
-	UE_AUTORTFM_API void RegisterOnAbortFromTheOpen(TTask<void()>&& Work);
-	UE_AUTORTFM_API void RedirectedLoad(uint32_t AddressSpace, void* DestPointer, uint64_t Size, uint64_t SourceAddress);
-	UE_AUTORTFM_API void RedirectedStore(uint32_t AddressSpace, uint64_t DestAddress, uint64_t Size, const void* SourcePointer);
-} // namespace ForTheRuntime
+UE_AUTORTFM_API void CascadingAbortTransactionInternal();
+UE_AUTORTFM_API void CascadingRetryTransactionInternal();
+UE_AUTORTFM_API void OnCommitInternal(TTask<void()>&& Work);
+UE_AUTORTFM_API void OnPreAbortInternal(TTask<void()>&& Work);
+UE_AUTORTFM_API void OnAbortInternal(TTask<void()>&& Work);
+UE_AUTORTFM_API void OnCompleteInternal(TTask<void()>&& Work);
+UE_AUTORTFM_API void PushOnCommitHandlerInternal(const void* Key, TTask<void()>&& Work);
+UE_AUTORTFM_API void PopOnCommitHandlerInternal(const void* Key);
+UE_AUTORTFM_API void PopAllOnCommitHandlersInternal(const void* Key);
+UE_AUTORTFM_API void PushOnAbortHandlerInternal(const void* Key, TTask<void()>&& Work);
+UE_AUTORTFM_API void PopOnAbortHandlerInternal(const void* Key);
+UE_AUTORTFM_API void PopAllOnAbortHandlersInternal(const void* Key);
+UE_AUTORTFM_API void RegisterOnCommitFromTheOpen(TTask<void()>&& Work);
+UE_AUTORTFM_API void RegisterOnAbortFromTheOpen(TTask<void()>&& Work);
+UE_AUTORTFM_API void RedirectedLoad(uint32_t AddressSpace, void* DestPointer, uint64_t Size, uint64_t SourceAddress);
+UE_AUTORTFM_API void RedirectedStore(uint32_t AddressSpace, uint64_t DestAddress, uint64_t Size, const void* SourcePointer);
+}  // namespace ForTheRuntime
 #endif
 
-template<typename FunctorType>
-void AutoRTFMFunctorInvoker(void* Arg) { (*static_cast<const FunctorType*>(Arg))(); }
+template <typename FunctorType>
+void AutoRTFMFunctorInvoker(void* Arg)
+{
+	(*static_cast<const FunctorType*>(Arg))();
+}
 
 #if UE_AUTORTFM
-template<typename FunctorType>
-auto AutoRTFMLookupInstrumentedFunctorInvoker(const FunctorType& Functor) -> void(*)(void*)
+template <typename FunctorType>
+auto AutoRTFMLookupInstrumentedFunctorInvoker(const FunctorType& Functor) -> void (*)(void*)
 {
 	// keep this as a single expression to help ensure that even Debug builds optimize this.
 	// if we put intermediate results in local variables then the compiler emits loads
 	// and stores to the stack which confuse our custom pass that tries to strip away
 	// the actual call to autortfm_lookup_function
-	void (*Result)(void*) = reinterpret_cast<void(*)(void*)>(autortfm_lookup_function(reinterpret_cast<void*>(&AutoRTFMFunctorInvoker<FunctorType>), "AutoRTFMLookupInstrumentedFunctorInvoker"));
+	void (*Result)(void*) = reinterpret_cast<void (*)(void*)>(autortfm_lookup_function(
+		reinterpret_cast<void*>(&AutoRTFMFunctorInvoker<FunctorType>), "AutoRTFMLookupInstrumentedFunctorInvoker"));
 	return Result;
 }
 #else
-template<typename FunctorType>
-auto AutoRTFMLookupInstrumentedFunctorInvoker(const FunctorType& Functor) -> void(*)(void*)
+template <typename FunctorType>
+auto AutoRTFMLookupInstrumentedFunctorInvoker(const FunctorType& Functor) -> void (*)(void*)
 {
 	return nullptr;
 }
@@ -282,33 +293,46 @@ UE_AUTORTFM_CRITICAL_INLINE void Initialize(const FExternAPI& ExternAPI)
 #if UE_AUTORTFM
 UE_AUTORTFM_API void Shutdown();
 #else
-UE_AUTORTFM_CRITICAL_INLINE void Shutdown()
-{
-}
+UE_AUTORTFM_CRITICAL_INLINE void Shutdown() {}
 #endif
 
 // Tells if we are currently running in a transaction. This will return true in an open nest
 // (see `Open`). This function is handled specially in the compiler and will be constant folded
 // as true in closed code, or preserved as a function call in open code.
-UE_AUTORTFM_CRITICAL_INLINE_ALWAYS bool IsTransactional() { return autortfm_is_transactional(); }
+UE_AUTORTFM_CRITICAL_INLINE_ALWAYS bool IsTransactional()
+{
+	return autortfm_is_transactional();
+}
 
 // Tells if we are currently running in the closed nest of a transaction. By default,
 // transactional code is in a closed nest; the only way to be in an open nest is to request it
 // via `Open`. This function is handled specially in the compiler and will be constant folded
 // as true in closed code, and false in open code.
-UE_AUTORTFM_CRITICAL_INLINE_ALWAYS bool IsClosed() { return autortfm_is_closed(); }
+UE_AUTORTFM_CRITICAL_INLINE_ALWAYS bool IsClosed()
+{
+	return autortfm_is_closed();
+}
 
 // Tells us if we are currently committing a transaction. This will return true
 // when inside an on-commit handler.
-UE_AUTORTFM_CRITICAL_INLINE bool IsCommitting() { return autortfm_is_committing(); }
+UE_AUTORTFM_CRITICAL_INLINE bool IsCommitting()
+{
+	return autortfm_is_committing();
+}
 
 // Tells us if we are currently committing or aborting a transaction. This will return true
 // when inside an on-abort, on-commit or on-complete handler.
-UE_AUTORTFM_CRITICAL_INLINE bool IsCommittingOrAborting() { return autortfm_is_committing_or_aborting(); }
+UE_AUTORTFM_CRITICAL_INLINE bool IsCommittingOrAborting()
+{
+	return autortfm_is_committing_or_aborting();
+}
 
 // Tells us if we are currently retrying a transaction. This will return true when inside an
 // on-abort or on-complete handler that was triggered by CascadingRetryTransaction.
-UE_AUTORTFM_CRITICAL_INLINE bool IsRetrying() { return autortfm_is_retrying(); }
+UE_AUTORTFM_CRITICAL_INLINE bool IsRetrying()
+{
+	return autortfm_is_retrying();
+}
 
 // Returns true if the passed-in pointer is on the stack of the currently-executing transaction.
 // This is occasionally necessary when writing OnAbort handlers for objects on the stack, since
@@ -332,19 +356,16 @@ UE_AUTORTFM_CRITICAL_INLINE TransactionID CurrentTransactionID()
 // if this nested transaction succeeds.
 //
 // If AutoRTFM is disabled, the code will be ran non-transactionally.
-template<typename FunctorType>
+template <typename FunctorType>
 UE_AUTORTFM_CRITICAL_INLINE ETransactionResult Transact(AUTORTFM_IMPLICIT_ENABLE const FunctorType& Functor)
 {
 #if UE_AUTORTFM_ENABLED
 	static constexpr autortfm_mode FunctorAutoRTFMMode = CallMode<FunctorType>;
-	static_assert(FunctorAutoRTFMMode == autortfm_mode_enable, "AutoRTFM mode of functor passed to AutoRTFM::Transact() is not autortfm_mode_enable");
+	static_assert(FunctorAutoRTFMMode == autortfm_mode_enable,
+		"AutoRTFM mode of functor passed to AutoRTFM::Transact() is not autortfm_mode_enable");
 #endif
-	ETransactionResult Result =
-		static_cast<ETransactionResult>(
-			autortfm_transact(
-				&AutoRTFMFunctorInvoker<FunctorType>,
-				AutoRTFMLookupInstrumentedFunctorInvoker<FunctorType>(Functor),
-				const_cast<void*>(static_cast<const void*>(&Functor))));
+	ETransactionResult Result = static_cast<ETransactionResult>(autortfm_transact(&AutoRTFMFunctorInvoker<FunctorType>,
+		AutoRTFMLookupInstrumentedFunctorInvoker<FunctorType>(Functor), const_cast<void*>(static_cast<const void*>(&Functor))));
 
 	return Result;
 }
@@ -353,14 +374,11 @@ UE_AUTORTFM_CRITICAL_INLINE ETransactionResult Transact(AUTORTFM_IMPLICIT_ENABLE
 // The reason we expose it is that it allows the caller's module to not
 // be compiled with the AutoRTFM instrumentation of functions if the only
 // thing that's being invoked is a function in the open.
-template<typename FunctorType>
+template <typename FunctorType>
 UE_AUTORTFM_CRITICAL_INLINE ETransactionResult TransactThenOpen(const FunctorType& Functor)
 {
-	ETransactionResult Result =
-		static_cast<ETransactionResult>(
-			autortfm_transact_then_open(
-				&AutoRTFMFunctorInvoker<FunctorType>,
-				const_cast<void*>(static_cast<const void*>(&Functor))));
+	ETransactionResult Result = static_cast<ETransactionResult>(
+		autortfm_transact_then_open(&AutoRTFMFunctorInvoker<FunctorType>, const_cast<void*>(static_cast<const void*>(&Functor))));
 
 	return Result;
 }
@@ -399,12 +417,8 @@ enum class EOpenFlags
 // a transaction or not. Returns the value returned by Functor.
 // ReturnType must be void or a type that can be safely copied from the open to a closed transaction.
 // TAssignFromOpenToClosed must have a specialization for the type that is being returned.
-template
-<
-	EOpenFlags Flags = EOpenFlags::Default,
-	typename FunctorType = void,
-	typename ReturnType = decltype(std::declval<FunctorType>()())
->
+template <EOpenFlags Flags = EOpenFlags::Default, typename FunctorType = void,
+	typename ReturnType = decltype(std::declval<FunctorType>()())>
 UE_AUTORTFM_CRITICAL_INLINE ReturnType Open(const FunctorType& Functor)
 {
 #if UE_AUTORTFM
@@ -415,21 +429,22 @@ UE_AUTORTFM_CRITICAL_INLINE ReturnType Open(const FunctorType& Functor)
 // Macro filth to either call autortfm_open() or autortfm_open_no_sanitize()
 // depending on Flags. Cannot be done as a lambda or inner class method as
 // any call indirection will break compiler optimizations.
-#define AUTORTFM_INTERNAL_CALL_IN_OPEN(Arg, Function) \
-	do { \
+#define AUTORTFM_INTERNAL_CALL_IN_OPEN(Arg, Function)  \
+	do                                                 \
+	{                                                  \
 		if constexpr (Flags == EOpenFlags::NoSanitize) \
-		{ \
-			autortfm_open_no_sanitize(Function, Arg); \
-		} \
-		else \
-		{ \
-			autortfm_open(Function, Arg); \
-		} \
+		{                                              \
+			autortfm_open_no_sanitize(Function, Arg);  \
+		}                                              \
+		else                                           \
+		{                                              \
+			autortfm_open(Function, Arg);              \
+		}                                              \
 	} while (false)
 
 	if constexpr (std::is_same_v<void, ReturnType>)
 	{
-		auto Call = [] AUTORTFM_DISABLE (void* ArgPtr)
+		auto Call = [] AUTORTFM_DISABLE(void* ArgPtr)
 		{
 			auto& Functor = *reinterpret_cast<FunctorType*>(ArgPtr);
 			UE_AUTORTFM_CALLSITE_FORCEINLINE Functor();
@@ -441,8 +456,7 @@ UE_AUTORTFM_CRITICAL_INLINE ReturnType Open(const FunctorType& Functor)
 		static constexpr EReturnFromOpenMode Mode = ReturnFromOpenModeFor<ReturnType>;
 		static_assert(Mode != EReturnFromOpenMode::Unsupported, "function return type is not safe to return from Open()");
 
-		if constexpr (Mode == EReturnFromOpenMode::CopyConstructInClosed ||
-		              Mode == EReturnFromOpenMode::MoveConstructInClosed)
+		if constexpr (Mode == EReturnFromOpenMode::CopyConstructInClosed || Mode == EReturnFromOpenMode::MoveConstructInClosed)
 		{
 			std::aligned_storage_t<sizeof(ReturnType), alignof(ReturnType)> OpenReturnValueMemory;
 
@@ -454,10 +468,10 @@ UE_AUTORTFM_CRITICAL_INLINE ReturnType Open(const FunctorType& Functor)
 			};
 			FArgs FunctorAndReturnValue{Functor, &OpenReturnValueMemory};
 			{
-				auto Call = [] AUTORTFM_DISABLE (void* ArgPtr)
+				auto Call = [] AUTORTFM_DISABLE(void* ArgPtr)
 				{
 					FArgs& Args = *reinterpret_cast<FArgs*>(ArgPtr);
-					new (Args.OpenReturnValueMemory) ReturnType (Args.Functor());
+					new (Args.OpenReturnValueMemory) ReturnType(Args.Functor());
 				};
 				AUTORTFM_INTERNAL_CALL_IN_OPEN(&FunctorAndReturnValue, Call);
 			}
@@ -465,14 +479,14 @@ UE_AUTORTFM_CRITICAL_INLINE ReturnType Open(const FunctorType& Functor)
 			// Re-interpret the OpenReturnValueMemory as a ReturnType&.
 			ReturnType& OpenReturnValue = *std::launder(reinterpret_cast<ReturnType*>(&OpenReturnValueMemory));
 
-			if constexpr(Mode == EReturnFromOpenMode::CopyConstructInClosed)
+			if constexpr (Mode == EReturnFromOpenMode::CopyConstructInClosed)
 			{
 				// Copy-construct the return value in the closed.
 				ReturnType ClosedReturnValue{OpenReturnValue};
 				// Destruct OpenReturnValue in the open, if required.
 				if constexpr (!std::is_trivially_destructible_v<ReturnType>)
 				{
-					auto Destruct = [] AUTORTFM_DISABLE (void* ArgPtr)
+					auto Destruct = [] AUTORTFM_DISABLE(void* ArgPtr)
 					{
 						reinterpret_cast<ReturnType*>(ArgPtr)->~ReturnType();
 					};
@@ -487,7 +501,7 @@ UE_AUTORTFM_CRITICAL_INLINE ReturnType Open(const FunctorType& Functor)
 				// Destruct OpenReturnValue in the open, if required.
 				if constexpr (!std::is_trivially_destructible_v<ReturnType>)
 				{
-					auto Destruct = [] AUTORTFM_DISABLE (void* ArgPtr)
+					auto Destruct = [] AUTORTFM_DISABLE(void* ArgPtr)
 					{
 						reinterpret_cast<ReturnType*>(ArgPtr)->~ReturnType();
 					};
@@ -504,7 +518,7 @@ UE_AUTORTFM_CRITICAL_INLINE ReturnType Open(const FunctorType& Functor)
 				ReturnType ReturnValue;
 			};
 			FArgs Args{Functor};
-			auto Call = [] AUTORTFM_DISABLE (void* ArgPtr)
+			auto Call = [] AUTORTFM_DISABLE(void* ArgPtr)
 			{
 				FArgs& Args = *reinterpret_cast<FArgs*>(ArgPtr);
 				ReturnType::AutoRTFMAssignFromOpenToClosed(Args.ReturnValue, Args.Functor());
@@ -518,9 +532,9 @@ UE_AUTORTFM_CRITICAL_INLINE ReturnType Open(const FunctorType& Functor)
 		}
 	}
 #undef AUTORTFM_INTERNAL_CALL_IN_OPEN
-#else // UE_AUTORTFM
+#else   // UE_AUTORTFM
 	return Functor();
-#endif // UE_AUTORTFM
+#endif  // UE_AUTORTFM
 }
 
 // Always executes the given code transactionally when called from a transaction nest
@@ -532,23 +546,23 @@ UE_AUTORTFM_CRITICAL_INLINE ReturnType Open(const FunctorType& Functor)
 // attempting to use the transaction is undefined behaviour. The caller should
 // return to the closed as quickly as possible to avoid the risk of the
 // transaction being used in its rolled-back state.
-template<typename FunctorType> [[nodiscard]] UE_AUTORTFM_CRITICAL_INLINE EContextStatus Close(AUTORTFM_IMPLICIT_ENABLE const FunctorType& Functor)
+template <typename FunctorType>
+[[nodiscard]] UE_AUTORTFM_CRITICAL_INLINE EContextStatus Close(AUTORTFM_IMPLICIT_ENABLE const FunctorType& Functor)
 {
 #if UE_AUTORTFM_ENABLED
 	static constexpr autortfm_mode FunctorAutoRTFMMode = CallMode<FunctorType>;
-	static_assert(FunctorAutoRTFMMode == autortfm_mode_enable, "AutoRTFM mode of functor passed to AutoRTFM::Close() is not autortfm_mode_enable");
+	static_assert(FunctorAutoRTFMMode == autortfm_mode_enable,
+		"AutoRTFM mode of functor passed to AutoRTFM::Close() is not autortfm_mode_enable");
 #endif
-	return static_cast<EContextStatus>(
-		autortfm_close(
-			&AutoRTFMFunctorInvoker<FunctorType>,
-			AutoRTFMLookupInstrumentedFunctorInvoker<FunctorType>(Functor),
-			const_cast<void*>(static_cast<const void*>(&Functor))));
+	return static_cast<EContextStatus>(autortfm_close(&AutoRTFMFunctorInvoker<FunctorType>,
+		AutoRTFMLookupInstrumentedFunctorInvoker<FunctorType>(Functor), const_cast<void*>(static_cast<const void*>(&Functor))));
 }
 
 // Similar to Close(), but can be called from outside of a transaction.
 // If called from outside a transaction then CloseIfTransactional() will
 // call Functor and always return EContextStatus::OnTrack.
-template<typename FunctorType> [[nodiscard]] UE_AUTORTFM_CRITICAL_INLINE EContextStatus CloseIfTransactional(AUTORTFM_IMPLICIT_ENABLE FunctorType&& Functor)
+template <typename FunctorType>
+[[nodiscard]] UE_AUTORTFM_CRITICAL_INLINE EContextStatus CloseIfTransactional(AUTORTFM_IMPLICIT_ENABLE FunctorType&& Functor)
 {
 #if UE_AUTORTFM_ENABLED
 	if (IsTransactional())
@@ -579,9 +593,7 @@ UE_AUTORTFM_CRITICAL_INLINE void CascadingRetryTransaction()
 	}
 }
 #else
-UE_AUTORTFM_CRITICAL_INLINE void CascadingRetryTransaction()
-{
-}
+UE_AUTORTFM_CRITICAL_INLINE void CascadingRetryTransaction() {}
 #endif
 
 #if UE_AUTORTFM
@@ -590,7 +602,8 @@ UE_AUTORTFM_CRITICAL_INLINE void CascadingRetryTransaction()
 // at that point, the worklist is run in FIFO order.
 // If this is called outside a transaction or from an open nest, then the work
 // happens immediately.
-template<typename FunctorType> UE_AUTORTFM_CRITICAL_INLINE void OnCommit(AUTORTFM_IMPLICIT_DISABLE FunctorType&& Work)
+template <typename FunctorType>
+UE_AUTORTFM_CRITICAL_INLINE void OnCommit(AUTORTFM_IMPLICIT_DISABLE FunctorType&& Work)
 {
 	if (autortfm_is_closed())
 	{
@@ -602,14 +615,19 @@ template<typename FunctorType> UE_AUTORTFM_CRITICAL_INLINE void OnCommit(AUTORTF
 	}
 }
 #else
-template<typename FunctorType> UE_AUTORTFM_CRITICAL_INLINE void OnCommit(FunctorType&& Work) { Work(); }
+template <typename FunctorType>
+UE_AUTORTFM_CRITICAL_INLINE void OnCommit(FunctorType&& Work)
+{
+	Work();
+}
 #endif
 
 #if UE_AUTORTFM
 // Have some work happen when this transaction aborts (before memory rollback).
 // If an abort occurs, the work list is run in LIFO order.
 // If this is called outside a transaction or from an open nest then the work is ignored.
-template<typename FunctorType> UE_AUTORTFM_CRITICAL_INLINE void OnPreAbort(AUTORTFM_IMPLICIT_DISABLE FunctorType&& Work)
+template <typename FunctorType>
+UE_AUTORTFM_CRITICAL_INLINE void OnPreAbort(AUTORTFM_IMPLICIT_DISABLE FunctorType&& Work)
 {
 	if (autortfm_is_closed())
 	{
@@ -617,14 +635,18 @@ template<typename FunctorType> UE_AUTORTFM_CRITICAL_INLINE void OnPreAbort(AUTOR
 	}
 }
 #else
-template<typename FunctorType> UE_AUTORTFM_CRITICAL_INLINE void OnPreAbort(FunctorType&&) {}
+template <typename FunctorType>
+UE_AUTORTFM_CRITICAL_INLINE void OnPreAbort(FunctorType&&)
+{
+}
 #endif
 
 #if UE_AUTORTFM
 // Have some work happen when this transaction aborts (after memory rollback).
 // If an abort occurs, the work list is run in LIFO order.
 // If this is called outside a transaction or from an open nest then the work is ignored.
-template<typename FunctorType> UE_AUTORTFM_CRITICAL_INLINE void OnAbort(AUTORTFM_IMPLICIT_DISABLE FunctorType&& Work)
+template <typename FunctorType>
+UE_AUTORTFM_CRITICAL_INLINE void OnAbort(AUTORTFM_IMPLICIT_DISABLE FunctorType&& Work)
 {
 	if (autortfm_is_closed())
 	{
@@ -632,7 +654,10 @@ template<typename FunctorType> UE_AUTORTFM_CRITICAL_INLINE void OnAbort(AUTORTFM
 	}
 }
 #else
-template<typename FunctorType> UE_AUTORTFM_CRITICAL_INLINE void OnAbort(FunctorType&&) {}
+template <typename FunctorType>
+UE_AUTORTFM_CRITICAL_INLINE void OnAbort(FunctorType&&)
+{
+}
 #endif
 
 #if UE_AUTORTFM
@@ -641,7 +666,8 @@ template<typename FunctorType> UE_AUTORTFM_CRITICAL_INLINE void OnAbort(FunctorT
 // the end of the transaction.
 // The worklist is run in FIFO order, after the OnCommit or OnAbort worklist has finished.
 // If this is called outside a transaction or from an open nest, then the work is ignored.
-template<typename FunctorType> UE_AUTORTFM_CRITICAL_INLINE void OnComplete(AUTORTFM_IMPLICIT_DISABLE FunctorType&& Work)
+template <typename FunctorType>
+UE_AUTORTFM_CRITICAL_INLINE void OnComplete(AUTORTFM_IMPLICIT_DISABLE FunctorType&& Work)
 {
 	if (autortfm_is_closed())
 	{
@@ -649,7 +675,10 @@ template<typename FunctorType> UE_AUTORTFM_CRITICAL_INLINE void OnComplete(AUTOR
 	}
 }
 #else
-template<typename FunctorType> UE_AUTORTFM_CRITICAL_INLINE void OnComplete(FunctorType&&) {}
+template <typename FunctorType>
+UE_AUTORTFM_CRITICAL_INLINE void OnComplete(FunctorType&&)
+{
+}
 #endif
 
 #if UE_AUTORTFM
@@ -657,7 +686,8 @@ template<typename FunctorType> UE_AUTORTFM_CRITICAL_INLINE void OnComplete(Funct
 // the handler can be unregistered (see `PopOnCommitHandler`). This is useful
 // for scoped mutations that need an abort handler present unless execution
 // reaches the end of the relevant scope.
-template<typename FunctorType> UE_AUTORTFM_CRITICAL_INLINE void PushOnCommitHandler(const void* Key, AUTORTFM_IMPLICIT_DISABLE FunctorType&& Work)
+template <typename FunctorType>
+UE_AUTORTFM_CRITICAL_INLINE void PushOnCommitHandler(const void* Key, AUTORTFM_IMPLICIT_DISABLE FunctorType&& Work)
 {
 	if (autortfm_is_closed())
 	{
@@ -669,7 +699,10 @@ template<typename FunctorType> UE_AUTORTFM_CRITICAL_INLINE void PushOnCommitHand
 // the handler can be unregistered (see `PopOnCommitHandler`). This is useful
 // for scoped mutations that need an abort handler present unless execution
 // reaches the end of the relevant scope.
-template<typename FunctorType> UE_AUTORTFM_CRITICAL_INLINE void PushOnCommitHandler(const void*, AUTORTFM_IMPLICIT_DISABLE FunctorType&&) {}
+template <typename FunctorType>
+UE_AUTORTFM_CRITICAL_INLINE void PushOnCommitHandler(const void*, AUTORTFM_IMPLICIT_DISABLE FunctorType&&)
+{
+}
 #endif
 
 #if UE_AUTORTFM
@@ -705,7 +738,8 @@ UE_AUTORTFM_CRITICAL_INLINE void PopAllOnCommitHandlers(const void*) {}
 // the handler can be unregistered (see `PopOnAbortHandler`). This is useful
 // for scoped mutations that need an abort handler present unless execution
 // reaches the end of the relevant scope.
-template<typename FunctorType> UE_AUTORTFM_CRITICAL_INLINE void PushOnAbortHandler(const void* Key, AUTORTFM_IMPLICIT_DISABLE FunctorType&& Work)
+template <typename FunctorType>
+UE_AUTORTFM_CRITICAL_INLINE void PushOnAbortHandler(const void* Key, AUTORTFM_IMPLICIT_DISABLE FunctorType&& Work)
 {
 	if (autortfm_is_closed())
 	{
@@ -717,7 +751,10 @@ template<typename FunctorType> UE_AUTORTFM_CRITICAL_INLINE void PushOnAbortHandl
 // the handler can be unregistered (see `PopOnAbortHandler`). This is useful
 // for scoped mutations that need an abort handler present unless execution
 // reaches the end of the relevant scope.
-template<typename FunctorType> UE_AUTORTFM_CRITICAL_INLINE void PushOnAbortHandler(const void* Key, FunctorType&&) {}
+template <typename FunctorType>
+UE_AUTORTFM_CRITICAL_INLINE void PushOnAbortHandler(const void* Key, FunctorType&&)
+{
+}
 #endif
 
 #if UE_AUTORTFM
@@ -764,9 +801,7 @@ struct FHeapRedirectCallbacks
 #if UE_AUTORTFM
 void RegisterHeapRedirectCallbacks(FHeapRedirectCallbacks Callbacks);
 #else
-UE_AUTORTFM_CRITICAL_INLINE void RegisterHeapRedirectCallbacks(FHeapRedirectCallbacks /*Callbacks*/)
-{
-}
+UE_AUTORTFM_CRITICAL_INLINE void RegisterHeapRedirectCallbacks(FHeapRedirectCallbacks /*Callbacks*/) {}
 #endif
 
 // Inform the runtime that we have performed a new object allocation. It's only
@@ -777,13 +812,13 @@ UE_AUTORTFM_CRITICAL_INLINE void RegisterHeapRedirectCallbacks(FHeapRedirectCall
 // to some nice optimizations. This does nothing when called from open code.
 UE_AUTORTFM_CRITICAL_INLINE void* DidAllocate(void* Ptr, size_t Size)
 {
-    return autortfm_did_allocate(Ptr, Size);
+	return autortfm_did_allocate(Ptr, Size);
 }
 
 // Inform the runtime that we have free'd a given memory location.
 UE_AUTORTFM_CRITICAL_INLINE void DidFree(void* Ptr)
 {
-    autortfm_did_free(Ptr);
+	autortfm_did_free(Ptr);
 }
 
 // Informs the runtime that a block of memory is about to be overwritten in the open.
@@ -801,7 +836,8 @@ UE_AUTORTFM_CRITICAL_INLINE void RecordOpenWrite(void* Ptr, size_t Size, EWriteF
 // a possible abort. Normally, tracking memory overwrites should be automatically
 // handled by AutoRTFM, but manual overwrite tracking may be required for third-party
 // libraries or outside compilers (such as ISPC).
-template<typename Type> UE_AUTORTFM_CRITICAL_INLINE void RecordOpenWrite(Type* Ptr, EWriteFlags Flags = EWriteFlags::Default)
+template <typename Type>
+UE_AUTORTFM_CRITICAL_INLINE void RecordOpenWrite(Type* Ptr, EWriteFlags Flags = EWriteFlags::Default)
 {
 	autortfm_record_open_write_with_flags(Ptr, sizeof(Type), static_cast<autortfm_write_flags>(Flags | EWriteFlags::NoSanitize));
 }
@@ -810,7 +846,7 @@ template<typename Type> UE_AUTORTFM_CRITICAL_INLINE void RecordOpenWrite(Type* P
 // of Ref before changing it to NewValue.
 // Can only be called in the open.
 #if UE_AUTORTFM_ENABLED
-template<typename RefType, typename ValueType>
+template <typename RefType, typename ValueType>
 AUTORTFM_DISABLE_NO_SANITIZE void Assign(RefType& Ref, ValueType&& NewValue, EWriteFlags Flags = EWriteFlags::Default)
 {
 	autortfm_record_open_write_with_flags(&Ref, sizeof(Ref), static_cast<autortfm_write_flags>(Flags));
@@ -820,7 +856,7 @@ AUTORTFM_DISABLE_NO_SANITIZE void Assign(RefType& Ref, ValueType&& NewValue, EWr
 	}
 }
 #else
-template<typename RefType, typename ValueType>
+template <typename RefType, typename ValueType>
 UE_AUTORTFM_CRITICAL_INLINE void Assign(RefType& Ref, ValueType&& NewValue, EWriteFlags Flags = EWriteFlags::Default)
 {
 	UE_AUTORTFM_UNUSED(Flags);
@@ -830,7 +866,7 @@ UE_AUTORTFM_CRITICAL_INLINE void Assign(RefType& Ref, ValueType&& NewValue, EWri
 
 // A RAII helper to call AutoRTFM::Assign(Variable, NewValue) on construction
 // and again on destruction with the old value of Variable.
-template<typename DataType>
+template <typename DataType>
 struct AUTORTFM_DISABLE_NO_SANITIZE TScopedAssign
 {
 	TScopedAssign(DataType& Variable, DataType NewValue) : Variable{Variable}, OldValue{Variable}
@@ -841,9 +877,10 @@ struct AUTORTFM_DISABLE_NO_SANITIZE TScopedAssign
 	{
 		Assign(Variable, OldValue);
 	}
+
 private:
 	TScopedAssign(DataType& Variable) = delete;
-	const TScopedAssign& operator = (DataType& Variable) = delete;
+	const TScopedAssign& operator=(DataType& Variable) = delete;
 	DataType& Variable;
 	DataType OldValue;
 };
@@ -883,285 +920,287 @@ UE_AUTORTFM_CRITICAL_INLINE_ALWAYS void UnreachableIfClosed(const char* Message
 // A collection of power-user functions that are reserved for use by the AutoRTFM runtime only.
 namespace ForTheRuntime
 {
-	// An enum to represent the various ways we want to enable/disable the AutoRTFM runtime.
-	// This enum has effective groups of functionality such that if a higher priority group
-	// has enabled or disabled the runtime, a lower priority group cannot then override that.
-	//
-	// We have from higher to lower priority:
-	// - Forced Enabled/Disabled - used by CVars when force enabling/disabling AutoRTFM.
-	// - Override Enabled/Disabled - override any setting of enabled/disabled as was set by a CVar.
-	// - Enabled/Disabled - used by CVars when enabling/disabling AutoRTFM.
-	// - Default Enabled/Disabled - whether we should be enabled or disabled by default (used for different backend executables).
-	//
-	// For example the following would be valid:
-	// - At compile time the state is compiled in as default disabled.
-	// - The CVar is set to enabled, so we switch the state to enabled.
-	// - At runtime we detect a mode where we want AutoRTFM and try to switch the default to enabled, but the CVar already enabled it so this is ignored.
-	// - Then for a given codepath we override AutoRTFM to disabled so we switch the state to disabled.
-	enum EAutoRTFMEnabledState
-	{
-		// Disable AutoRTFM.
-		AutoRTFM_Disabled = 0,
+// An enum to represent the various ways we want to enable/disable the AutoRTFM runtime.
+// This enum has effective groups of functionality such that if a higher priority group
+// has enabled or disabled the runtime, a lower priority group cannot then override that.
+//
+// We have from higher to lower priority:
+// - Forced Enabled/Disabled - used by CVars when force enabling/disabling AutoRTFM.
+// - Override Enabled/Disabled - override any setting of enabled/disabled as was set by a CVar.
+// - Enabled/Disabled - used by CVars when enabling/disabling AutoRTFM.
+// - Default Enabled/Disabled - whether we should be enabled or disabled by default (used for different backend executables).
+//
+// For example the following would be valid:
+// - At compile time the state is compiled in as default disabled.
+// - The CVar is set to enabled, so we switch the state to enabled.
+// - At runtime we detect a mode where we want AutoRTFM and try to switch the default to enabled, but the CVar already enabled it so
+// this is ignored.
+// - Then for a given codepath we override AutoRTFM to disabled so we switch the state to disabled.
+enum EAutoRTFMEnabledState
+{
+	// Disable AutoRTFM.
+	AutoRTFM_Disabled = 0,
 
-		// Enable AutoRTFM.
-		AutoRTFM_Enabled,
+	// Enable AutoRTFM.
+	AutoRTFM_Enabled,
 
-		// Force disable AutoRTFM.
-		AutoRTFM_ForcedDisabled,
+	// Force disable AutoRTFM.
+	AutoRTFM_ForcedDisabled,
 
-		// Force enable AutoRTFM.
-		AutoRTFM_ForcedEnabled,
+	// Force enable AutoRTFM.
+	AutoRTFM_ForcedEnabled,
 
-		// Whether our default is to be disabled.
-		AutoRTFM_DisabledByDefault,
+	// Whether our default is to be disabled.
+	AutoRTFM_DisabledByDefault,
 
-		// Whether our default is to be enabled.
-		AutoRTFM_EnabledByDefault,
+	// Whether our default is to be enabled.
+	AutoRTFM_EnabledByDefault,
 
-		// Whether we've overridden and AutoRTFM is disabled.
-		AutoRTFM_OverriddenDisabled,
+	// Whether we've overridden and AutoRTFM is disabled.
+	AutoRTFM_OverriddenDisabled,
 
-		// Whether we've overridden and AutoRTFM is enabled.
-		AutoRTFM_OverriddenEnabled,
-	};
+	// Whether we've overridden and AutoRTFM is enabled.
+	AutoRTFM_OverriddenEnabled,
+};
 
-	// An enum to represent whether we should abort and retry transactions (for testing purposes).
-	enum EAutoRTFMRetryTransactionState
-	{
-		// Do not abort and retry transactions (the default).
-		NoRetry = 0,
+// An enum to represent whether we should abort and retry transactions (for testing purposes).
+enum EAutoRTFMRetryTransactionState
+{
+	// Do not abort and retry transactions (the default).
+	NoRetry = 0,
 
-		// Abort and retry non-nested transactions (EG. only abort the parent transactional nest).
-		RetryNonNested,
+	// Abort and retry non-nested transactions (EG. only abort the parent transactional nest).
+	RetryNonNested,
 
-		// Abort and retry nested-transactions too. Will be slower as each nested-transaction will
-		// be aborted and retried at least *twice* (once when the non-nested transaction runs the
-		// first time, and a second time when the non-nested transaction is doing its retry after
-		// aborting).
-		RetryNestedToo,
-	};
+	// Abort and retry nested-transactions too. Will be slower as each nested-transaction will
+	// be aborted and retried at least *twice* (once when the non-nested transaction runs the
+	// first time, and a second time when the non-nested transaction is doing its retry after
+	// aborting).
+	RetryNestedToo,
+};
 
-	enum EAutoRTFMInternalAbortActionState
-	{
-		// Crash the process if we hit an internal AutoRTFM abort.
-		Crash = 0,
+enum EAutoRTFMInternalAbortActionState
+{
+	// Crash the process if we hit an internal AutoRTFM abort.
+	Crash = 0,
 
-		// Just do a normal transaction abort and let the runtime recover (used to test aborting codepaths).
-		Abort,
-	};
+	// Just do a normal transaction abort and let the runtime recover (used to test aborting codepaths).
+	Abort,
+};
 
-	// Set whether the AutoRTFM runtime is enabled or disabled. Returns true when the state was changed
-	// successfully.
-	UE_AUTORTFM_API bool SetAutoRTFMRuntime(EAutoRTFMEnabledState State);
+// Set whether the AutoRTFM runtime is enabled or disabled. Returns true when the state was changed
+// successfully.
+UE_AUTORTFM_API bool SetAutoRTFMRuntime(EAutoRTFMEnabledState State);
 
-	UE_AUTORTFM_API bool IsAutoRTFMRuntimeEnabledInternal();
+UE_AUTORTFM_API bool IsAutoRTFMRuntimeEnabledInternal();
 
-	// Query whether the AutoRTFM runtime is enabled.
-	UE_AUTORTFM_CRITICAL_INLINE bool IsAutoRTFMRuntimeEnabled()
+// Query whether the AutoRTFM runtime is enabled.
+UE_AUTORTFM_CRITICAL_INLINE bool IsAutoRTFMRuntimeEnabled()
+{
+	// If we are already in the closed nest of a transaction, we must have our runtime enabled!
+	if (AutoRTFM::IsClosed())
 	{
-		// If we are already in the closed nest of a transaction, we must have our runtime enabled!
-		if (AutoRTFM::IsClosed())
-		{
-			return true;
-		}
-
-		return IsAutoRTFMRuntimeEnabledInternal();
+		return true;
 	}
 
-	// Set the percentage [0..100] chance that a call to `CoinTossDisable` will end up disabling AutoRTFM.
-	// 100% means never disable via coin-toss, 0% means always disable. So passing `0.1` means disable
-	// all but 1/1000's calls via `CoinTossDisable`.
-	UE_AUTORTFM_API void SetAutoRTFMEnabledProbability(float Chance);
+	return IsAutoRTFMRuntimeEnabledInternal();
+}
 
-	// Get the enabled probability set via `SetAutoRTFMEnabledProbability`.
-	UE_AUTORTFM_API float GetAutoRTFMEnabledProbability();
+// Set the percentage [0..100] chance that a call to `CoinTossDisable` will end up disabling AutoRTFM.
+// 100% means never disable via coin-toss, 0% means always disable. So passing `0.1` means disable
+// all but 1/1000's calls via `CoinTossDisable`.
+UE_AUTORTFM_API void SetAutoRTFMEnabledProbability(float Chance);
 
-	// Call to randomly disable AutoRTFM with a probability set with `SetAutoRTFMEnabledProbability`.
-	// Returns true if AutoRTFM was disabled by this call.
-	UE_AUTORTFM_API bool CoinTossDisable();
+// Get the enabled probability set via `SetAutoRTFMEnabledProbability`.
+UE_AUTORTFM_API float GetAutoRTFMEnabledProbability();
 
-	UE_AUTORTFM_API void SetInternalAbortAction(EAutoRTFMInternalAbortActionState State);
+// Call to randomly disable AutoRTFM with a probability set with `SetAutoRTFMEnabledProbability`.
+// Returns true if AutoRTFM was disabled by this call.
+UE_AUTORTFM_API bool CoinTossDisable();
 
-	UE_AUTORTFM_API EAutoRTFMInternalAbortActionState GetInternalAbortAction();
+UE_AUTORTFM_API void SetInternalAbortAction(EAutoRTFMInternalAbortActionState State);
 
-	UE_AUTORTFM_API bool GetEnsureOnInternalAbort();
-	UE_AUTORTFM_API void SetEnsureOnInternalAbort(bool bEnabled);
+UE_AUTORTFM_API EAutoRTFMInternalAbortActionState GetInternalAbortAction();
 
-	UE_AUTORTFM_API autortfm_log_severity GetAutoRTFMHazardLogSeverity();
-	UE_AUTORTFM_API void SetAutoRTFMHazardLogSeverity(autortfm_log_severity Severity);
+UE_AUTORTFM_API bool GetEnsureOnInternalAbort();
+UE_AUTORTFM_API void SetEnsureOnInternalAbort(bool bEnabled);
 
-	// Set whether we should trigger an ensure on an abort-by-language.
-	[[deprecated("Use `SetEnsureOnInternalAbort` instead!")]]
-	inline void SetEnsureOnAbortByLanguage(bool bEnabled)
-	{
-		SetEnsureOnInternalAbort(bEnabled);
-	}
+UE_AUTORTFM_API autortfm_log_severity GetAutoRTFMHazardLogSeverity();
+UE_AUTORTFM_API void SetAutoRTFMHazardLogSeverity(autortfm_log_severity Severity);
 
-	// Returns whether the runtime will trigger an ensure on an abort-by-language, or not.
-	[[deprecated("Use `GetEnsureOnInternalAbort` instead!")]]
-	inline bool IsEnsureOnAbortByLanguageEnabled()
-	{
-		return GetEnsureOnInternalAbort();
-	}
+// Set whether we should trigger an ensure on an abort-by-language.
+[[deprecated("Use `SetEnsureOnInternalAbort` instead!")]] inline void SetEnsureOnAbortByLanguage(bool bEnabled)
+{
+	SetEnsureOnInternalAbort(bEnabled);
+}
 
-	// Returns whether we want to assert or ensure on a Language Error
-	[[deprecated("Use `GetInternalAbortAction` instead!")]]
-	inline bool IsAutoRTFMAssertOnError()
-	{
-		return EAutoRTFMInternalAbortActionState::Crash == GetInternalAbortAction();
-	}
+// Returns whether the runtime will trigger an ensure on an abort-by-language, or not.
+[[deprecated("Use `GetEnsureOnInternalAbort` instead!")]] inline bool IsEnsureOnAbortByLanguageEnabled()
+{
+	return GetEnsureOnInternalAbort();
+}
+
+// Returns whether we want to assert or ensure on a Language Error
+[[deprecated("Use `GetInternalAbortAction` instead!")]] inline bool IsAutoRTFMAssertOnError()
+{
+	return EAutoRTFMInternalAbortActionState::Crash == GetInternalAbortAction();
+}
+
+// Set whether we should retry transactions.
+UE_AUTORTFM_API void SetRetryTransaction(EAutoRTFMRetryTransactionState State);
+
+// Returns whether we should retry transactions.
+UE_AUTORTFM_API EAutoRTFMRetryTransactionState GetRetryTransaction();
+
+// Returns true if we should retry non-nested transactions.
+UE_AUTORTFM_API bool ShouldRetryNonNestedTransactions();
+
+// Returns true if we should also retry nested transactions.
+UE_AUTORTFM_API bool ShouldRetryNestedTransactionsToo();
+
+// Manually create a new transaction from open code and push it as a transaction nest.
+// Can only be called within an already active parent transaction (EG. this cannot start
+// a transaction nest itself).
+AUTORTFM_DISABLE UE_AUTORTFM_API void StartTransaction();
+
+// Manually commit the top transaction nest, popping it from the execution scope.
+// Can only be called within an already active parent transaction (EG. this cannot end
+// a transaction nest itself).
+AUTORTFM_DISABLE UE_AUTORTFM_API ETransactionResult CommitTransaction();
+
+// Manually clear the status of a user abort from the top transaction in a nest.
+AUTORTFM_DISABLE UE_AUTORTFM_API void ClearTransactionStatus();
+
+// Returns the current transaction status.
+AUTORTFM_DISABLE UE_AUTORTFM_API EContextStatus GetContextStatus();
+
+// RollbackTransaction() aborts the current transaction.
+// If the transaction is scoped, then this is kept on the transaction stack
+// until we return to the closed. Attempting to use this transaction before
+// transitioning to the closed is undefined behaviour.
+// If the transaction is unscoped, then this is popped immediately.
+// RollbackTransaction() sets the transaction status to AbortedByRequest.
+AUTORTFM_DISABLE UE_AUTORTFM_API void RollbackTransaction();
+
+// CascadingAbortRollbackTransaction() rolls back the current scoped
+// transaction, which must be a scoped transaction (i.e. created by
+// AutoRTFM::Transact() and not by StartTransaction()).
+// CascadingAbortRollbackTransaction() sets the transaction status to
+// AbortedByCascadingAbort, which will cause the next open -> closed
+// transition point to abort all transactions on the transaction nest.
+// Attempting to use any transactions in the open while unwinding is
+// undefined behaviour.
+AUTORTFM_DISABLE UE_AUTORTFM_API void CascadingAbortRollbackTransaction();
+
+// CascadingRetryRollbackTransaction() rolls back the current scoped
+// transaction, which must be a scoped transaction (i.e. created by
+// AutoRTFM::Transact() and not by StartTransaction()).
+// CascadingRetryRollbackTransaction() sets the transaction status to
+// AbortedByCascadingRetry, which will cause the next open -> closed
+// transition point to abort all transactions on the transaction nest,
+// and then retry the outermost transaction.
+// Attempting to use any transactions in the open while unwinding is
+// undefined behaviour.
+AUTORTFM_DISABLE UE_AUTORTFM_API void CascadingRetryRollbackTransaction();
+
+// Reserved for future.
+UE_AUTORTFM_CRITICAL_INLINE void RecordOpenRead(void const*, size_t) {}
+
+// Reserved for future.
+template <typename Type>
+UE_AUTORTFM_CRITICAL_INLINE void RecordOpenRead(Type*)
+{
+}
 
-	// Set whether we should retry transactions.
-	UE_AUTORTFM_API void SetRetryTransaction(EAutoRTFMRetryTransactionState State);
-
-	// Returns whether we should retry transactions.
-	UE_AUTORTFM_API EAutoRTFMRetryTransactionState GetRetryTransaction();
-
-	// Returns true if we should retry non-nested transactions.
-	UE_AUTORTFM_API bool ShouldRetryNonNestedTransactions();
-
-	// Returns true if we should also retry nested transactions.
-	UE_AUTORTFM_API bool ShouldRetryNestedTransactionsToo();
-
-	// Manually create a new transaction from open code and push it as a transaction nest.
-	// Can only be called within an already active parent transaction (EG. this cannot start
-	// a transaction nest itself).
-	AUTORTFM_DISABLE UE_AUTORTFM_API void StartTransaction();
-
-	// Manually commit the top transaction nest, popping it from the execution scope.
-	// Can only be called within an already active parent transaction (EG. this cannot end
-	// a transaction nest itself).
-	AUTORTFM_DISABLE UE_AUTORTFM_API ETransactionResult CommitTransaction();
-
-	// Manually clear the status of a user abort from the top transaction in a nest.
-	AUTORTFM_DISABLE UE_AUTORTFM_API void ClearTransactionStatus();
-
-	// Returns the current transaction status.
-	AUTORTFM_DISABLE UE_AUTORTFM_API EContextStatus GetContextStatus();
-
-	// RollbackTransaction() aborts the current transaction.
-	// If the transaction is scoped, then this is kept on the transaction stack
-	// until we return to the closed. Attempting to use this transaction before
-	// transitioning to the closed is undefined behaviour.
-	// If the transaction is unscoped, then this is popped immediately.
-	// RollbackTransaction() sets the transaction status to AbortedByRequest.
-	AUTORTFM_DISABLE UE_AUTORTFM_API void RollbackTransaction();
-
-	// CascadingAbortRollbackTransaction() rolls back the current scoped
-	// transaction, which must be a scoped transaction (i.e. created by
-	// AutoRTFM::Transact() and not by StartTransaction()).
-	// CascadingAbortRollbackTransaction() sets the transaction status to
-	// AbortedByCascadingAbort, which will cause the next open -> closed
-	// transition point to abort all transactions on the transaction nest.
-	// Attempting to use any transactions in the open while unwinding is
-	// undefined behaviour.
-	AUTORTFM_DISABLE UE_AUTORTFM_API void CascadingAbortRollbackTransaction();
-
-	// CascadingRetryRollbackTransaction() rolls back the current scoped
-	// transaction, which must be a scoped transaction (i.e. created by
-	// AutoRTFM::Transact() and not by StartTransaction()).
-	// CascadingRetryRollbackTransaction() sets the transaction status to
-	// AbortedByCascadingRetry, which will cause the next open -> closed
-	// transition point to abort all transactions on the transaction nest,
-	// and then retry the outermost transaction.
-	// Attempting to use any transactions in the open while unwinding is
-	// undefined behaviour.
-	AUTORTFM_DISABLE UE_AUTORTFM_API void CascadingRetryRollbackTransaction();
-
-	// Reserved for future.
-	UE_AUTORTFM_CRITICAL_INLINE void RecordOpenRead(void const*, size_t) {}
-
-	// Reserved for future.
-	template<typename Type> UE_AUTORTFM_CRITICAL_INLINE void RecordOpenRead(Type*) {}
-
-} // namespace ForTheRuntime
-
-} // namespace AutoRTFM
+}  // namespace ForTheRuntime
+
+}  // namespace AutoRTFM
 
 // Macro-based variants so we completely compile away when not in use, even in debug builds
 #if UE_AUTORTFM
 
 namespace AutoRTFM::Private
 {
-	struct FOpenHelper
+struct FOpenHelper
+{
+	template <typename FunctorType>
+	UE_AUTORTFM_FORCEINLINE void operator+(FunctorType&& F)
 	{
-		template<typename FunctorType>
-		UE_AUTORTFM_FORCEINLINE void operator+(FunctorType&& F)
-		{
-			AutoRTFM::Open(std::forward<FunctorType>(F));
-		}
-	};
-	struct FOpenNoSanitizeHelper
+		AutoRTFM::Open(std::forward<FunctorType>(F));
+	}
+};
+struct FOpenNoSanitizeHelper
+{
+	template <typename FunctorType>
+	UE_AUTORTFM_FORCEINLINE void operator+(FunctorType&& F)
 	{
-		template<typename FunctorType>
-		UE_AUTORTFM_FORCEINLINE void operator+(FunctorType&& F)
-		{
-			AutoRTFM::Open<EOpenFlags::NoSanitize>(std::forward<FunctorType>(F));
-		}
-	};
-	struct FCloseHelper
+		AutoRTFM::Open<EOpenFlags::NoSanitize>(std::forward<FunctorType>(F));
+	}
+};
+struct FCloseHelper
+{
+	template <typename FunctorType>
+	[[nodiscard]] EContextStatus operator+(FunctorType F)
 	{
-		template<typename FunctorType>
-		[[nodiscard]] EContextStatus operator+(FunctorType F)
-		{
-			return AutoRTFM::Close(std::move(F));
-		}
-	};
-	struct FOnPreAbortHelper
+		return AutoRTFM::Close(std::move(F));
+	}
+};
+struct FOnPreAbortHelper
+{
+	template <typename FunctorType>
+	UE_AUTORTFM_FORCEINLINE void operator+(FunctorType&& F)
 	{
-		template<typename FunctorType>
-		UE_AUTORTFM_FORCEINLINE void operator+(FunctorType&& F)
-		{
-			AutoRTFM::OnPreAbort(std::forward<FunctorType>(F));
-		}
-	};
-	struct FOnAbortHelper
+		AutoRTFM::OnPreAbort(std::forward<FunctorType>(F));
+	}
+};
+struct FOnAbortHelper
+{
+	template <typename FunctorType>
+	UE_AUTORTFM_FORCEINLINE void operator+(FunctorType&& F)
 	{
-		template<typename FunctorType>
-		UE_AUTORTFM_FORCEINLINE void operator+(FunctorType&& F)
-		{
-			AutoRTFM::OnAbort(std::forward<FunctorType>(F));
-		}
-	};
-	struct FOnCommitHelper
+		AutoRTFM::OnAbort(std::forward<FunctorType>(F));
+	}
+};
+struct FOnCommitHelper
+{
+	template <typename FunctorType>
+	UE_AUTORTFM_FORCEINLINE void operator+(FunctorType&& F)
 	{
-		template<typename FunctorType>
-		UE_AUTORTFM_FORCEINLINE void operator+(FunctorType&& F)
-		{
-			AutoRTFM::OnCommit(std::forward<FunctorType>(F));
-		}
-	};
-	struct FTransactHelper
+		AutoRTFM::OnCommit(std::forward<FunctorType>(F));
+	}
+};
+struct FTransactHelper
+{
+	template <typename FunctorType>
+	UE_AUTORTFM_FORCEINLINE void operator+(FunctorType&& F)
 	{
-		template<typename FunctorType>
-		UE_AUTORTFM_FORCEINLINE void operator+(FunctorType&& F)
-		{
-			AutoRTFM::Transact(std::forward<FunctorType>(F));
-		}
-	};
-	namespace /* must have internal linkage */
+		AutoRTFM::Transact(std::forward<FunctorType>(F));
+	}
+};
+namespace /* must have internal linkage */
+{
+struct FThreadLocalHelper
+{
+	template <typename Type, int Unique>
+	UE_AUTORTFM_ALWAYS_OPEN static Type& Get()
 	{
-		struct FThreadLocalHelper
-		{
-			template <typename Type, int Unique>
-			UE_AUTORTFM_ALWAYS_OPEN static Type& Get()
-			{
-				thread_local Type Data;
-				return Data;
-			}
-		};
+		thread_local Type Data;
+		return Data;
 	}
-} // namespace AutoRTFM::Private
-
-#define UE_AUTORTFM_DECLARE_THREAD_LOCAL_VAR_IMPL(Type, Name) Type& Name = ::AutoRTFM::Private::FThreadLocalHelper::Get<Type, __COUNTER__>()
-
-#define UE_AUTORTFM_OPEN_IMPL          ::AutoRTFM::Private::FOpenHelper{} + [&] AUTORTFM_DISABLE ()
-#define UE_AUTORTFM_OPEN_NO_SANITIZE_IMPL ::AutoRTFM::Private::FOpenNoSanitizeHelper{} + [&] AUTORTFM_DISABLE_NO_SANITIZE ()
-#define UE_AUTORTFM_CLOSE_IMPL         ::AutoRTFM::Private::FCloseHelper{} + [&] AUTORTFM_ENABLE ()
-#define UE_AUTORTFM_ONPREABORT_IMPL(...)  ::AutoRTFM::Private::FOnPreAbortHelper{} + [__VA_ARGS__] AUTORTFM_DISABLE () mutable
-#define UE_AUTORTFM_ONABORT_IMPL(...)  ::AutoRTFM::Private::FOnAbortHelper{} + [__VA_ARGS__] AUTORTFM_DISABLE () mutable
-#define UE_AUTORTFM_ONCOMMIT_IMPL(...) ::AutoRTFM::Private::FOnCommitHelper{} + [__VA_ARGS__] AUTORTFM_DISABLE () mutable
-#define UE_AUTORTFM_TRANSACT_IMPL      ::AutoRTFM::Private::FTransactHelper{} + [&] AUTORTFM_ENABLE ()
+};
+}
+}  // namespace AutoRTFM::Private
+
+#define UE_AUTORTFM_DECLARE_THREAD_LOCAL_VAR_IMPL(Type, Name) \
+	Type& Name = ::AutoRTFM::Private::FThreadLocalHelper::Get<Type, __COUNTER__>()
+
+#define UE_AUTORTFM_OPEN_IMPL ::AutoRTFM::Private::FOpenHelper{} + [&] AUTORTFM_DISABLE()
+#define UE_AUTORTFM_OPEN_NO_SANITIZE_IMPL ::AutoRTFM::Private::FOpenNoSanitizeHelper{} + [&] AUTORTFM_DISABLE_NO_SANITIZE()
+#define UE_AUTORTFM_CLOSE_IMPL ::AutoRTFM::Private::FCloseHelper{} + [&] AUTORTFM_ENABLE()
+#define UE_AUTORTFM_ONPREABORT_IMPL(...) ::AutoRTFM::Private::FOnPreAbortHelper{} + [__VA_ARGS__] AUTORTFM_DISABLE() mutable
+#define UE_AUTORTFM_ONABORT_IMPL(...) ::AutoRTFM::Private::FOnAbortHelper{} + [__VA_ARGS__] AUTORTFM_DISABLE() mutable
+#define UE_AUTORTFM_ONCOMMIT_IMPL(...) ::AutoRTFM::Private::FOnCommitHelper{} + [__VA_ARGS__] AUTORTFM_DISABLE() mutable
+#define UE_AUTORTFM_TRANSACT_IMPL ::AutoRTFM::Private::FTransactHelper{} + [&] AUTORTFM_ENABLE()
 #else
 
 // Do nothing, these should be followed by blocks that should be either executed or not executed
@@ -1200,29 +1239,30 @@ namespace AutoRTFM::Private
 // In non-transactional code paths the block of code will not be executed at all.
 // The macro arguments are the capture specification for the lambda.
 // Calls should be written like this: UE_AUTORTFM_ONPREABORT(=) { ... code ... };
-#define UE_AUTORTFM_ONPREABORT(...)  UE_AUTORTFM_ONPREABORT_IMPL(__VA_ARGS__)
+#define UE_AUTORTFM_ONPREABORT(...) UE_AUTORTFM_ONPREABORT_IMPL(__VA_ARGS__)
 
 // Runs a block of code if a transaction aborts (after memory rollback).
 // In non-transactional code paths the block of code will not be executed at all.
 // The macro arguments are the capture specification for the lambda.
 // Calls should be written like this: UE_AUTORTFM_ONABORT(=) { ... code ... };
-#define UE_AUTORTFM_ONABORT(...)  UE_AUTORTFM_ONABORT_IMPL(__VA_ARGS__)
+#define UE_AUTORTFM_ONABORT(...) UE_AUTORTFM_ONABORT_IMPL(__VA_ARGS__)
 
 // Runs a block of code if a transaction commits successfully.
 // In non-transactional code paths the block of code will be executed immediately.
 // The macro arguments are the capture specification for the lambda.
 // Calls should be written like this: UE_AUTORTFM_ONCOMMIT(=) { ... code ... };
-#define UE_AUTORTFM_ONCOMMIT(...)  UE_AUTORTFM_ONCOMMIT_IMPL(__VA_ARGS__)
+#define UE_AUTORTFM_ONCOMMIT(...) UE_AUTORTFM_ONCOMMIT_IMPL(__VA_ARGS__)
 
 // Runs a block of code in the closed, transactionally, within a new transaction.
 // Calls should be written like this: UE_AUTORTFM_TRANSACT { ... code ... };
-#define UE_AUTORTFM_TRANSACT  UE_AUTORTFM_TRANSACT_IMPL
+#define UE_AUTORTFM_TRANSACT UE_AUTORTFM_TRANSACT_IMPL
 
 #if UE_AUTORTFM
-#define UE_AUTORTFM_REGISTER_OPEN_TO_CLOSED_FUNCTIONS(...) \
-	static const ::AutoRTFM::ForTheRuntime::TAutoRegisterOpenToClosedFunctions<__VA_ARGS__> UE_AUTORTFM_CONCAT(AutoRTFMFunctionRegistration, __COUNTER__)
+#define UE_AUTORTFM_REGISTER_OPEN_TO_CLOSED_FUNCTIONS(...)                                                      \
+	static const ::AutoRTFM::ForTheRuntime::TAutoRegisterOpenToClosedFunctions<__VA_ARGS__> UE_AUTORTFM_CONCAT( \
+		AutoRTFMFunctionRegistration, __COUNTER__)
 #else
 #define UE_AUTORTFM_REGISTER_OPEN_TO_CLOSED_FUNCTIONS(...)
 #endif
 
-#endif // __cplusplus
+#endif  // __cplusplus
diff --git a/Engine/Source/Runtime/AutoRTFM/Public/AutoRTFM/CAPI.h b/Engine/Source/Runtime/AutoRTFM/Public/AutoRTFM/CAPI.h
--- a/Engine/Source/Runtime/AutoRTFM/Public/AutoRTFM/CAPI.h
+++ b/Engine/Source/Runtime/AutoRTFM/Public/AutoRTFM/CAPI.h
@@ -5,19 +5,20 @@
 #include "AutoRTFM/Defines.h"
 #include "Constants.h"
 
-#include <cstdarg> // for va_list
-#include <cstddef> // for size_t
+#include <cstdarg>  // for va_list
+#include <cstddef>  // for size_t
 
 // Define UE_AUTORTFM_API (DLL export / import attributes)
-#if defined(__cplusplus) && defined(__UNREAL__) && !(defined(UE_AUTORTFM_DO_NOT_INCLUDE_PLATFORM_H) && UE_AUTORTFM_DO_NOT_INCLUDE_PLATFORM_H)
-	// Include HAL/Platform.h for DLLIMPORT / DLLEXPORT definitions, which
-	// UBT can use as a definition for AUTORTFM_API.
-	#include <HAL/Platform.h>
-	#define UE_AUTORTFM_API AUTORTFM_API
+#if defined(__cplusplus) && defined(__UNREAL__) \
+	&& !(defined(UE_AUTORTFM_DO_NOT_INCLUDE_PLATFORM_H) && UE_AUTORTFM_DO_NOT_INCLUDE_PLATFORM_H)
+// Include HAL/Platform.h for DLLIMPORT / DLLEXPORT definitions, which
+// UBT can use as a definition for AUTORTFM_API.
+#include <HAL/Platform.h>
+#define UE_AUTORTFM_API AUTORTFM_API
 #else
-	#ifndef UE_AUTORTFM_API
-		#define UE_AUTORTFM_API
-	#endif
+#ifndef UE_AUTORTFM_API
+#define UE_AUTORTFM_API
+#endif
 #endif
 
 #ifdef __cplusplus
@@ -25,172 +26,174 @@ extern "C"
 {
 #endif
 
-// The C API exists for a few reasons:
-//
-// - It makes linking easy. AutoRTFM has to deal with a weird kind of linking
-//   where the compiler directly emits calls to functions with a given name.
-//   It's easiest to do that in llvm if the functions have C linkage and C ABI.
-// - It makes testing easy. Even seemingly simple C++ code introduces pitfalls
-//   for AutoRTFM. So very focused tests work best when written in C.
-// - It makes compiler optimizations much easier as there is no mangling to
-// 	 consider when looking for functions in the runtime we can optimize.
-//
-// We use snake_case for C API surface area to make it easy to distinguish.
-//
-// The C API should not be used directly - it is here purely as an
-// implementation detail.
-
-// This must match AutoRTFM::ETransactionResult.
-typedef enum
-{
-    autortfm_aborted_by_request = 0,
-    autortfm_aborted_by_language,
-    autortfm_committed,
-	autortfm_aborted_by_transact_during_commit,
-	autortfm_aborted_by_transact_during_abort,
-	autortfm_aborted_by_cascade,
-} autortfm_result;
-
-// This must match AutoRTFM::EContextStatus.
-typedef enum
-{
-	autortfm_status_idle = 0,
-	autortfm_status_ontrack,
-	autortfm_status_aborted_by_failed_lock_aquisition,
-	autortfm_status_aborted_by_language,
-	autortfm_status_aborted_by_request,
-	autortfm_status_committing,
-	autortfm_status_aborted_by_cascading_abort,
-	autortfm_status_aborted_by_cascading_retry,
-	autortfm_status_in_static_local_initializer,
-	autortfm_status_in_post_abort
-} autortfm_status;
-
-// AutoRTFM logging severity.
-typedef enum
-{
-	autortfm_log_verbose = 0,
-	autortfm_log_info,
-	autortfm_log_warn,
-	autortfm_log_error,
-	autortfm_log_fatal,
-} autortfm_log_severity;
-
-// Flags that can be passed to autortfm_record_open_write_with_flags().
-typedef enum
-{
-	// No flags applied to the write.
-	autortfm_write_default = 0x0,
-	// The write should not be considered by the AutoRTFM sanitizer.
-	autortfm_write_no_sanitize = 0x1,
-	// The autortfm_extern_api::RollbackWrite will be called to perform
-	// the rollback for this write if the transaction is aborted.
-	autortfm_write_custom_rollback = 0x2,
-	// Custom flag bits which can be used by the custom abort callbacks.
-	// Requires autortfm_write_custom_rollback for usage.
-	autortfm_write_custom_flag_0 = 0x4,
-	autortfm_write_custom_flag_1 = 0x8,
-} autortfm_write_flags;
-
-// An opaque unique identifier for a transaction.
-typedef uint64_t autortfm_transaction_id;
-
-// A single frame of a callstack capture.
-typedef unsigned long long autortfm_callstack_frame;
-
-// Function pointers used by AutoRTFM for heap allocations, etc.
-typedef struct
-{
-	// The function used to allocate memory from the heap.
-	// Must not be null.
-	void* (*Allocate)(size_t Size, size_t Alignment);
-
-	// The function used to reallocate memory from the heap.
-	// Must not be null.
-	void* (*Reallocate)(void* Pointer, size_t Size, size_t Alignment);
-
-	// The function used to allocate zeroed memory from the heap.
-	// Must not be null.
-	void* (*AllocateZeroed)(size_t Size, size_t Alignment);
-
-	// The function used to free memory allocated by Allocate() and AllocateZeroed().
-	// Must not be null.
-	void (*Free)(void* Pointer);
-
-	// Function used to log messages using a printf-style format string and va_list arguments.
-	// Strings use UTF-8 encoding.
-	// Must not be null.
-	void (*Log)(const char* File, int Line, void* ProgramCounter, autortfm_log_severity Severity, const char* Format, va_list Args);
-
-	// Function used to log messages with a callstack using a printf-style format string and va_list arguments.
-	// Strings use UTF-8 encoding.
-	// Must not be null.
-	void (*LogWithCallstack)(void* ProgramCounter, autortfm_log_severity Severity, const char* Format, va_list Args);
-
-	// Function used to report an ensure failure using a printf-style format string and va_list arguments.
-	// Strings use UTF-8 encoding.
-	// Must not be null.
-	void (*EnsureFailure)(const char* File, int Line, void* ProgramCounter, const char* Condition, const char* Format, va_list Args);
-
-	// Function used to query whether a log severity is active.
-	// Must not be null.
-	bool (*IsLogActive)(autortfm_log_severity Severity);
-
-	// Performs a capture of the currently executing callstack, writing at most
-	// MaxCount frame addresses to StackOut.
-	// Returns the number of frame addresses actually written.
-	// Must not be null.
-	size_t (*CaptureCallstack)(size_t MaxCount, autortfm_callstack_frame* StackOut);
-
-	// Prints a callstack captured with CaptureCallstack to the log with the
-	// given severity.
-	// Must not be null.
-	void (*LogCallstack)(autortfm_log_severity Severity, size_t Count, autortfm_callstack_frame const* Stack);
-
-	// Optional callback to be informed when the value returned by
-	// ForTheRuntime::IsAutoRTFMRuntimeEnabled() changes.
-	// Can be null.
-	void (*OnRuntimeEnabledChanged)();
-
-	// Optional callback to be informed when the value returned by
-	// ForTheRuntime::GetRetryTransaction() changes.
-	// Can be null.
-	void (*OnRetryTransactionsChanged)();
+	// The C API exists for a few reasons:
+	//
+	// - It makes linking easy. AutoRTFM has to deal with a weird kind of linking
+	//   where the compiler directly emits calls to functions with a given name.
+	//   It's easiest to do that in llvm if the functions have C linkage and C ABI.
+	// - It makes testing easy. Even seemingly simple C++ code introduces pitfalls
+	//   for AutoRTFM. So very focused tests work best when written in C.
+	// - It makes compiler optimizations much easier as there is no mangling to
+	// 	 consider when looking for functions in the runtime we can optimize.
+	//
+	// We use snake_case for C API surface area to make it easy to distinguish.
+	//
+	// The C API should not be used directly - it is here purely as an
+	// implementation detail.
+
+	// This must match AutoRTFM::ETransactionResult.
+	typedef enum
+	{
+		autortfm_aborted_by_request = 0,
+		autortfm_aborted_by_language,
+		autortfm_committed,
+		autortfm_aborted_by_transact_during_commit,
+		autortfm_aborted_by_transact_during_abort,
+		autortfm_aborted_by_cascade,
+	} autortfm_result;
+
+	// This must match AutoRTFM::EContextStatus.
+	typedef enum
+	{
+		autortfm_status_idle = 0,
+		autortfm_status_ontrack,
+		autortfm_status_aborted_by_failed_lock_aquisition,
+		autortfm_status_aborted_by_language,
+		autortfm_status_aborted_by_request,
+		autortfm_status_committing,
+		autortfm_status_aborted_by_cascading_abort,
+		autortfm_status_aborted_by_cascading_retry,
+		autortfm_status_in_static_local_initializer,
+		autortfm_status_in_post_abort
+	} autortfm_status;
+
+	// AutoRTFM logging severity.
+	typedef enum
+	{
+		autortfm_log_verbose = 0,
+		autortfm_log_info,
+		autortfm_log_warn,
+		autortfm_log_error,
+		autortfm_log_fatal,
+	} autortfm_log_severity;
+
+	// Flags that can be passed to autortfm_record_open_write_with_flags().
+	typedef enum
+	{
+		// No flags applied to the write.
+		autortfm_write_default = 0x0,
+		// The write should not be considered by the AutoRTFM sanitizer.
+		autortfm_write_no_sanitize = 0x1,
+		// The autortfm_extern_api::RollbackWrite will be called to perform
+		// the rollback for this write if the transaction is aborted.
+		autortfm_write_custom_rollback = 0x2,
+		// Custom flag bits which can be used by the custom abort callbacks.
+		// Requires autortfm_write_custom_rollback for usage.
+		autortfm_write_custom_flag_0 = 0x4,
+		autortfm_write_custom_flag_1 = 0x8,
+	} autortfm_write_flags;
+
+	// An opaque unique identifier for a transaction.
+	typedef uint64_t autortfm_transaction_id;
+
+	// A single frame of a callstack capture.
+	typedef unsigned long long autortfm_callstack_frame;
+
+	// Function pointers used by AutoRTFM for heap allocations, etc.
+	typedef struct
+	{
+		// The function used to allocate memory from the heap.
+		// Must not be null.
+		void* (*Allocate)(size_t Size, size_t Alignment);
+
+		// The function used to reallocate memory from the heap.
+		// Must not be null.
+		void* (*Reallocate)(void* Pointer, size_t Size, size_t Alignment);
+
+		// The function used to allocate zeroed memory from the heap.
+		// Must not be null.
+		void* (*AllocateZeroed)(size_t Size, size_t Alignment);
+
+		// The function used to free memory allocated by Allocate() and AllocateZeroed().
+		// Must not be null.
+		void (*Free)(void* Pointer);
+
+		// Function used to log messages using a printf-style format string and va_list arguments.
+		// Strings use UTF-8 encoding.
+		// Must not be null.
+		void (*Log)(
+			const char* File, int Line, void* ProgramCounter, autortfm_log_severity Severity, const char* Format, va_list Args);
+
+		// Function used to log messages with a callstack using a printf-style format string and va_list arguments.
+		// Strings use UTF-8 encoding.
+		// Must not be null.
+		void (*LogWithCallstack)(void* ProgramCounter, autortfm_log_severity Severity, const char* Format, va_list Args);
+
+		// Function used to report an ensure failure using a printf-style format string and va_list arguments.
+		// Strings use UTF-8 encoding.
+		// Must not be null.
+		void (*EnsureFailure)(
+			const char* File, int Line, void* ProgramCounter, const char* Condition, const char* Format, va_list Args);
+
+		// Function used to query whether a log severity is active.
+		// Must not be null.
+		bool (*IsLogActive)(autortfm_log_severity Severity);
+
+		// Performs a capture of the currently executing callstack, writing at most
+		// MaxCount frame addresses to StackOut.
+		// Returns the number of frame addresses actually written.
+		// Must not be null.
+		size_t (*CaptureCallstack)(size_t MaxCount, autortfm_callstack_frame* StackOut);
+
+		// Prints a callstack captured with CaptureCallstack to the log with the
+		// given severity.
+		// Must not be null.
+		void (*LogCallstack)(autortfm_log_severity Severity, size_t Count, autortfm_callstack_frame const* Stack);
+
+		// Optional callback to be informed when the value returned by
+		// ForTheRuntime::IsAutoRTFMRuntimeEnabled() changes.
+		// Can be null.
+		void (*OnRuntimeEnabledChanged)();
+
+		// Optional callback to be informed when the value returned by
+		// ForTheRuntime::GetRetryTransaction() changes.
+		// Can be null.
+		void (*OnRetryTransactionsChanged)();
 
 #if AUTORTFM_SANITIZER
-	// Optional callback to be informed when the value returned by
-	// ForTheRuntime::GetSanitizerMode() changes.
-	// Can be null.
-	void (*OnSanitizerModeChanged)();
-	
-	// Optional callback to be informed when the value returned by
-	// ForTheRuntime::SanitizerRecordClosedWriteCallstacks() changes.
-	// Can be null.
-	void (*OnSanitizerRecordClosedWriteCallstacksChanged)();
-#endif // AUTORTFM_SANITIZER
-
-	// Custom abort handler for writes that that were added with the
-	// autortfm_write_custom_rollback flag. It is the responsibility of this
-	// function to perform the rollback (copy of Size bytes from Value to
-	// Address).
-	// RollbackWrite may contain multiple writes condensed into a single
-	// contigious block of memory, and the writes may be split into multiple
-	// calls with no defined size or alignment granularity.
-	// Calls are ordered starting with the most recently recorded writes.
-	// Address / Value pointers are ordered starting with the least recently
-	// recorded writes.
-	// Split writes are guaranteed to be "completed" with the next immediate
-	// call(s).
-	void (*RollbackWrite)(void* Address, const void* Value, size_t Size, autortfm_write_flags Flags);
-} autortfm_extern_api;
+		// Optional callback to be informed when the value returned by
+		// ForTheRuntime::GetSanitizerMode() changes.
+		// Can be null.
+		void (*OnSanitizerModeChanged)();
+
+		// Optional callback to be informed when the value returned by
+		// ForTheRuntime::SanitizerRecordClosedWriteCallstacks() changes.
+		// Can be null.
+		void (*OnSanitizerRecordClosedWriteCallstacksChanged)();
+#endif  // AUTORTFM_SANITIZER
+
+		// Custom abort handler for writes that that were added with the
+		// autortfm_write_custom_rollback flag. It is the responsibility of this
+		// function to perform the rollback (copy of Size bytes from Value to
+		// Address).
+		// RollbackWrite may contain multiple writes condensed into a single
+		// contigious block of memory, and the writes may be split into multiple
+		// calls with no defined size or alignment granularity.
+		// Calls are ordered starting with the most recently recorded writes.
+		// Address / Value pointers are ordered starting with the least recently
+		// recorded writes.
+		// Split writes are guaranteed to be "completed" with the next immediate
+		// call(s).
+		void (*RollbackWrite)(void* Address, const void* Value, size_t Size, autortfm_write_flags Flags);
+	} autortfm_extern_api;
 
 #if UE_AUTORTFM_ENABLED
-// Initialize the AutoRTFM library.
-// Must only be called once for the lifetime of the application.
-// Parameters:
-//   ExternAPI - Function pointers used by AutoRTFM for heap allocations, etc.
-//               Must be non-null.
-UE_AUTORTFM_API void autortfm_initialize(const autortfm_extern_api* ExternAPI) AUTORTFM_NOEXCEPT;
+	// Initialize the AutoRTFM library.
+	// Must only be called once for the lifetime of the application.
+	// Parameters:
+	//   ExternAPI - Function pointers used by AutoRTFM for heap allocations, etc.
+	//               Must be non-null.
+	UE_AUTORTFM_API void autortfm_initialize(const autortfm_extern_api* ExternAPI) AUTORTFM_NOEXCEPT;
 #else
 UE_AUTORTFM_CRITICAL_INLINE void autortfm_initialize(const autortfm_extern_api* ExternAPI) AUTORTFM_NOEXCEPT
 {
@@ -199,37 +202,35 @@ UE_AUTORTFM_CRITICAL_INLINE void autortfm_initialize(const autortfm_extern_api*
 #endif
 
 #if UE_AUTORTFM_ENABLED
-// Shutdown the AutoRTFM library.
-// Must only be called once for the lifetime of the application, after autortfm_initialize().
-UE_AUTORTFM_API void autortfm_shutdown() AUTORTFM_NOEXCEPT;
+	// Shutdown the AutoRTFM library.
+	// Must only be called once for the lifetime of the application, after autortfm_initialize().
+	UE_AUTORTFM_API void autortfm_shutdown() AUTORTFM_NOEXCEPT;
 #else
-UE_AUTORTFM_CRITICAL_INLINE void autortfm_shutdown() AUTORTFM_NOEXCEPT
-{
-}
+UE_AUTORTFM_CRITICAL_INLINE void autortfm_shutdown() AUTORTFM_NOEXCEPT {}
 #endif
 
 #if UE_AUTORTFM_ENABLED
-// Note: There is no implementation of this function.
-// The AutoRTFM compiler will replace all calls to this function with a constant boolean value.
-UE_AUTORTFM_API bool autortfm_is_closed(void) AUTORTFM_NOEXCEPT;
+	// Note: There is no implementation of this function.
+	// The AutoRTFM compiler will replace all calls to this function with a constant boolean value.
+	UE_AUTORTFM_API bool autortfm_is_closed(void) AUTORTFM_NOEXCEPT;
 #else
 UE_AUTORTFM_CRITICAL_INLINE bool autortfm_is_closed(void) AUTORTFM_NOEXCEPT
 {
-    return false;
+	return false;
 }
 #endif
 
 #if UE_AUTORTFM_ENABLED
-UE_AUTORTFM_API bool autortfm_is_transactional(void) AUTORTFM_NOEXCEPT;
+	UE_AUTORTFM_API bool autortfm_is_transactional(void) AUTORTFM_NOEXCEPT;
 #else
 UE_AUTORTFM_CRITICAL_INLINE bool autortfm_is_transactional(void) AUTORTFM_NOEXCEPT
 {
-    return false;
+	return false;
 }
 #endif
 
 #if UE_AUTORTFM_ENABLED
-UE_AUTORTFM_API bool autortfm_is_committing(void) AUTORTFM_NOEXCEPT;
+	UE_AUTORTFM_API bool autortfm_is_committing(void) AUTORTFM_NOEXCEPT;
 #else
 UE_AUTORTFM_CRITICAL_INLINE bool autortfm_is_committing(void) AUTORTFM_NOEXCEPT
 {
@@ -238,7 +239,7 @@ UE_AUTORTFM_CRITICAL_INLINE bool autortfm_is_committing(void) AUTORTFM_NOEXCEPT
 #endif
 
 #if UE_AUTORTFM_ENABLED
-UE_AUTORTFM_API bool autortfm_is_committing_or_aborting(void) AUTORTFM_NOEXCEPT;
+	UE_AUTORTFM_API bool autortfm_is_committing_or_aborting(void) AUTORTFM_NOEXCEPT;
 #else
 UE_AUTORTFM_CRITICAL_INLINE bool autortfm_is_committing_or_aborting(void) AUTORTFM_NOEXCEPT
 {
@@ -247,7 +248,7 @@ UE_AUTORTFM_CRITICAL_INLINE bool autortfm_is_committing_or_aborting(void) AUTORT
 #endif
 
 #if UE_AUTORTFM_ENABLED
-UE_AUTORTFM_API bool autortfm_is_retrying(void) AUTORTFM_NOEXCEPT;
+	UE_AUTORTFM_API bool autortfm_is_retrying(void) AUTORTFM_NOEXCEPT;
 #else
 UE_AUTORTFM_CRITICAL_INLINE bool autortfm_is_retrying(void) AUTORTFM_NOEXCEPT
 {
@@ -256,9 +257,11 @@ UE_AUTORTFM_CRITICAL_INLINE bool autortfm_is_retrying(void) AUTORTFM_NOEXCEPT
 #endif
 
 #if UE_AUTORTFM_ENABLED
-UE_AUTORTFM_API autortfm_result autortfm_transact(void (*UninstrumentedWork)(void*), void (*InstrumentedWork)(void*), void* Arg) AUTORTFM_EXCEPT;
+	UE_AUTORTFM_API autortfm_result autortfm_transact(
+		void (*UninstrumentedWork)(void*), void (*InstrumentedWork)(void*), void* Arg) AUTORTFM_EXCEPT;
 #else
-UE_AUTORTFM_CRITICAL_INLINE autortfm_result autortfm_transact(void (*UninstrumentedWork)(void*), void (*InstrumentedWork)(void*), void* Arg) AUTORTFM_EXCEPT
+UE_AUTORTFM_CRITICAL_INLINE autortfm_result autortfm_transact(
+	void (*UninstrumentedWork)(void*), void (*InstrumentedWork)(void*), void* Arg) AUTORTFM_EXCEPT
 {
 	UE_AUTORTFM_UNUSED(InstrumentedWork);
 	UninstrumentedWork(Arg);
@@ -267,50 +270,64 @@ UE_AUTORTFM_CRITICAL_INLINE autortfm_result autortfm_transact(void (*Uninstrumen
 #endif
 
 #if UE_AUTORTFM_ENABLED
-UE_AUTORTFM_API autortfm_result autortfm_transact_then_open(void (*UninstrumentedWork)(void*), void* Arg) AUTORTFM_EXCEPT;
+	UE_AUTORTFM_API autortfm_result autortfm_transact_then_open(void (*UninstrumentedWork)(void*), void* Arg) AUTORTFM_EXCEPT;
 #else
 UE_AUTORTFM_CRITICAL_INLINE autortfm_result autortfm_transact_then_open(void (*UninstrumentedWork)(void*), void* Arg) AUTORTFM_EXCEPT
 {
 	UninstrumentedWork(Arg);
-    return autortfm_committed;
+	return autortfm_committed;
 }
 #endif
 
 #if UE_AUTORTFM_ENABLED
-UE_AUTORTFM_API void autortfm_abort_transaction() AUTORTFM_NOEXCEPT;
+	UE_AUTORTFM_API void autortfm_abort_transaction() AUTORTFM_NOEXCEPT;
 #else
 UE_AUTORTFM_CRITICAL_INLINE void autortfm_abort_transaction() AUTORTFM_NOEXCEPT {}
 #endif
 
 #if UE_AUTORTFM_ENABLED
-AUTORTFM_OPEN UE_AUTORTFM_API autortfm_transaction_id autortfm_current_transaction_id() AUTORTFM_NOEXCEPT;
+	AUTORTFM_OPEN UE_AUTORTFM_API autortfm_transaction_id autortfm_current_transaction_id() AUTORTFM_NOEXCEPT;
 #else
-UE_AUTORTFM_CRITICAL_INLINE autortfm_transaction_id autortfm_current_transaction_id() AUTORTFM_NOEXCEPT { return 0; }
+UE_AUTORTFM_CRITICAL_INLINE autortfm_transaction_id autortfm_current_transaction_id() AUTORTFM_NOEXCEPT
+{
+	return 0;
+}
 #endif
 
 #if UE_AUTORTFM_ENABLED
-UE_AUTORTFM_API autortfm_result autortfm_commit_transaction() AUTORTFM_NOEXCEPT;
+	UE_AUTORTFM_API autortfm_result autortfm_commit_transaction() AUTORTFM_NOEXCEPT;
 #else
-UE_AUTORTFM_CRITICAL_INLINE autortfm_result autortfm_commit_transaction() AUTORTFM_NOEXCEPT { return autortfm_aborted_by_language; }
+UE_AUTORTFM_CRITICAL_INLINE autortfm_result autortfm_commit_transaction() AUTORTFM_NOEXCEPT
+{
+	return autortfm_aborted_by_language;
+}
 #endif
 
 #if UE_AUTORTFM_ENABLED
-UE_AUTORTFM_FORCENOINLINE UE_AUTORTFM_API void autortfm_open(void (*work)(void* arg), void* arg);
+	UE_AUTORTFM_FORCENOINLINE UE_AUTORTFM_API void autortfm_open(void (*work)(void* arg), void* arg);
 #else
-UE_AUTORTFM_CRITICAL_INLINE void autortfm_open(void (*work)(void* arg), void* arg) { work(arg); }
+UE_AUTORTFM_CRITICAL_INLINE void autortfm_open(void (*work)(void* arg), void* arg)
+{
+	work(arg);
+}
 #endif
 
 #if UE_AUTORTFM_ENABLED
-UE_AUTORTFM_FORCENOINLINE UE_AUTORTFM_API void autortfm_open_no_sanitize(void (*work)(void* arg), void* arg);
+	UE_AUTORTFM_FORCENOINLINE UE_AUTORTFM_API void autortfm_open_no_sanitize(void (*work)(void* arg), void* arg);
 #else
-UE_AUTORTFM_CRITICAL_INLINE void autortfm_open_no_sanitize(void (*work)(void* arg), void* arg) { work(arg); }
+UE_AUTORTFM_CRITICAL_INLINE void autortfm_open_no_sanitize(void (*work)(void* arg), void* arg)
+{
+	work(arg);
+}
 #endif
 
 #if UE_AUTORTFM_ENABLED
-[[nodiscard]] UE_AUTORTFM_API autortfm_status autortfm_close(void (*UninstrumentedWork)(void*), void (*InstrumentedWork)(void*), void* Arg);
+	[[nodiscard]] UE_AUTORTFM_API autortfm_status autortfm_close(
+		void (*UninstrumentedWork)(void*), void (*InstrumentedWork)(void*), void* Arg);
 #else
 AUTORTFM_DISABLE_UNREACHABLE_CODE_WARNINGS
-[[nodiscard]] UE_AUTORTFM_CRITICAL_INLINE autortfm_status autortfm_close(void (*UninstrumentedWork)(void*), void (*InstrumentedWork)(void*), void* Arg)
+[[nodiscard]] UE_AUTORTFM_CRITICAL_INLINE autortfm_status autortfm_close(
+	void (*UninstrumentedWork)(void*), void (*InstrumentedWork)(void*), void* Arg)
 {
 	UE_AUTORTFM_UNUSED(UninstrumentedWork);
 	UE_AUTORTFM_UNUSED(InstrumentedWork);
@@ -324,7 +341,7 @@ AUTORTFM_RESTORE_UNREACHABLE_CODE_WARNINGS
 #endif
 
 #if UE_AUTORTFM_ENABLED
-UE_AUTORTFM_API void autortfm_record_open_write(void* Ptr, size_t Size) AUTORTFM_NOEXCEPT;
+	UE_AUTORTFM_API void autortfm_record_open_write(void* Ptr, size_t Size) AUTORTFM_NOEXCEPT;
 #else
 UE_AUTORTFM_CRITICAL_INLINE void autortfm_record_open_write(void* Ptr, size_t Size) AUTORTFM_NOEXCEPT
 {
@@ -334,9 +351,10 @@ UE_AUTORTFM_CRITICAL_INLINE void autortfm_record_open_write(void* Ptr, size_t Si
 #endif
 
 #if UE_AUTORTFM_ENABLED
-UE_AUTORTFM_API void autortfm_record_open_write_with_flags(void* Ptr, size_t Size, autortfm_write_flags Flags) AUTORTFM_NOEXCEPT;
+	UE_AUTORTFM_API void autortfm_record_open_write_with_flags(void* Ptr, size_t Size, autortfm_write_flags Flags) AUTORTFM_NOEXCEPT;
 #else
-UE_AUTORTFM_CRITICAL_INLINE void autortfm_record_open_write_with_flags(void* Ptr, size_t Size, autortfm_write_flags Flags) AUTORTFM_NOEXCEPT
+UE_AUTORTFM_CRITICAL_INLINE void autortfm_record_open_write_with_flags(
+	void* Ptr, size_t Size, autortfm_write_flags Flags) AUTORTFM_NOEXCEPT
 {
 	UE_AUTORTFM_UNUSED(Ptr);
 	UE_AUTORTFM_UNUSED(Size);
@@ -344,31 +362,31 @@ UE_AUTORTFM_CRITICAL_INLINE void autortfm_record_open_write_with_flags(void* Ptr
 }
 #endif
 
-// autortfm_open_to_closed_mapping maps an open function to its closed variant.
-struct autortfm_open_to_closed_mapping
-{
-	void* Open;
-	void* Closed;
-};
-
-// autortfm_open_to_closed_table holds a pointer to a null-terminated list of
-// autortfm_open_to_closed_mapping, and an intrusive linked-list pointer to the
-// previous and next registered autortfm_open_to_closed_table.
-struct autortfm_open_to_closed_table
-{
-	// Null-terminated open function to closed function mapping table.
-	const struct autortfm_open_to_closed_mapping* Mappings;
-	// An intrusive linked-list pointer to the previous autortfm_open_to_closed_table.
-	// Used by autortfm_register_open_to_closed_functions().
-	struct autortfm_open_to_closed_table* Prev;
-	// An intrusive linked-list pointer to the next autortfm_open_to_closed_table.
-	// Used by autortfm_register_open_to_closed_functions().
-	struct autortfm_open_to_closed_table* Next;
-};
+	// autortfm_open_to_closed_mapping maps an open function to its closed variant.
+	struct autortfm_open_to_closed_mapping
+	{
+		void* Open;
+		void* Closed;
+	};
+
+	// autortfm_open_to_closed_table holds a pointer to a null-terminated list of
+	// autortfm_open_to_closed_mapping, and an intrusive linked-list pointer to the
+	// previous and next registered autortfm_open_to_closed_table.
+	struct autortfm_open_to_closed_table
+	{
+		// Null-terminated open function to closed function mapping table.
+		const struct autortfm_open_to_closed_mapping* Mappings;
+		// An intrusive linked-list pointer to the previous autortfm_open_to_closed_table.
+		// Used by autortfm_register_open_to_closed_functions().
+		struct autortfm_open_to_closed_table* Prev;
+		// An intrusive linked-list pointer to the next autortfm_open_to_closed_table.
+		// Used by autortfm_register_open_to_closed_functions().
+		struct autortfm_open_to_closed_table* Next;
+	};
 
 #if UE_AUTORTFM_ENABLED
-UE_AUTORTFM_API void autortfm_register_open_to_closed_functions(struct autortfm_open_to_closed_table* Table) AUTORTFM_NOEXCEPT;
-UE_AUTORTFM_API void autortfm_unregister_open_to_closed_functions(struct autortfm_open_to_closed_table* Table) AUTORTFM_NOEXCEPT;
+	UE_AUTORTFM_API void autortfm_register_open_to_closed_functions(struct autortfm_open_to_closed_table* Table) AUTORTFM_NOEXCEPT;
+	UE_AUTORTFM_API void autortfm_unregister_open_to_closed_functions(struct autortfm_open_to_closed_table* Table) AUTORTFM_NOEXCEPT;
 #else
 UE_AUTORTFM_CRITICAL_INLINE void autortfm_register_open_to_closed_functions(autortfm_open_to_closed_table* Table) AUTORTFM_NOEXCEPT
 {
@@ -381,7 +399,7 @@ UE_AUTORTFM_CRITICAL_INLINE void autortfm_unregister_open_to_closed_functions(au
 #endif
 
 #if UE_AUTORTFM_ENABLED
-UE_AUTORTFM_API bool autortfm_is_on_current_transaction_stack(void* Ptr) AUTORTFM_NOEXCEPT;
+	UE_AUTORTFM_API bool autortfm_is_on_current_transaction_stack(void* Ptr) AUTORTFM_NOEXCEPT;
 #else
 UE_AUTORTFM_CRITICAL_INLINE bool autortfm_is_on_current_transaction_stack(void* Ptr) AUTORTFM_NOEXCEPT
 {
@@ -391,16 +409,16 @@ UE_AUTORTFM_CRITICAL_INLINE bool autortfm_is_on_current_transaction_stack(void*
 #endif
 
 #if UE_AUTORTFM_ENABLED
-UE_AUTORTFM_API void autortfm_on_commit(void (*work)(void* arg), void* arg);
+	UE_AUTORTFM_API void autortfm_on_commit(void (*work)(void* arg), void* arg);
 #else
 UE_AUTORTFM_CRITICAL_INLINE void autortfm_on_commit(void (*work)(void* arg), void* arg)
 {
-    work(arg);
+	work(arg);
 }
 #endif
 
 #if UE_AUTORTFM_ENABLED
-UE_AUTORTFM_API void autortfm_on_pre_abort(void (*work)(void* arg), void* arg) AUTORTFM_NOEXCEPT;
+	UE_AUTORTFM_API void autortfm_on_pre_abort(void (*work)(void* arg), void* arg) AUTORTFM_NOEXCEPT;
 #else
 UE_AUTORTFM_CRITICAL_INLINE void autortfm_on_pre_abort(void (*work)(void* arg), void* arg) AUTORTFM_NOEXCEPT
 {
@@ -410,7 +428,7 @@ UE_AUTORTFM_CRITICAL_INLINE void autortfm_on_pre_abort(void (*work)(void* arg),
 #endif
 
 #if UE_AUTORTFM_ENABLED
-UE_AUTORTFM_API void autortfm_on_abort(void (*work)(void* arg), void* arg) AUTORTFM_NOEXCEPT;
+	UE_AUTORTFM_API void autortfm_on_abort(void (*work)(void* arg), void* arg) AUTORTFM_NOEXCEPT;
 #else
 UE_AUTORTFM_CRITICAL_INLINE void autortfm_on_abort(void (*work)(void* arg), void* arg) AUTORTFM_NOEXCEPT
 {
@@ -420,10 +438,11 @@ UE_AUTORTFM_CRITICAL_INLINE void autortfm_on_abort(void (*work)(void* arg), void
 #endif
 
 #if UE_AUTORTFM_ENABLED
-UE_AUTORTFM_API void autortfm_push_on_abort_handler(const void* key, void (*work)(void* arg), void* arg) AUTORTFM_NOEXCEPT;
-UE_AUTORTFM_API void autortfm_pop_on_abort_handler(const void* key) AUTORTFM_NOEXCEPT;
+	UE_AUTORTFM_API void autortfm_push_on_abort_handler(const void* key, void (*work)(void* arg), void* arg) AUTORTFM_NOEXCEPT;
+	UE_AUTORTFM_API void autortfm_pop_on_abort_handler(const void* key) AUTORTFM_NOEXCEPT;
 #else
-UE_AUTORTFM_CRITICAL_INLINE void autortfm_push_on_abort_handler(const void* key, void (*work)(void* arg), void* arg) AUTORTFM_NOEXCEPT
+UE_AUTORTFM_CRITICAL_INLINE void autortfm_push_on_abort_handler(
+	const void* key, void (*work)(void* arg), void* arg) AUTORTFM_NOEXCEPT
 {
 	UE_AUTORTFM_UNUSED(key);
 	UE_AUTORTFM_UNUSED(work);
@@ -437,17 +456,17 @@ UE_AUTORTFM_CRITICAL_INLINE void autortfm_pop_on_abort_handler(const void* key)
 #endif
 
 #if UE_AUTORTFM_ENABLED
-UE_AUTORTFM_API void* autortfm_did_allocate(void* ptr, size_t size) AUTORTFM_NOEXCEPT;
+	UE_AUTORTFM_API void* autortfm_did_allocate(void* ptr, size_t size) AUTORTFM_NOEXCEPT;
 #else
 UE_AUTORTFM_CRITICAL_INLINE void* autortfm_did_allocate(void* ptr, size_t size) AUTORTFM_NOEXCEPT
 {
 	UE_AUTORTFM_UNUSED(size);
-    return ptr;
+	return ptr;
 }
 #endif
 
 #if UE_AUTORTFM_ENABLED
-UE_AUTORTFM_API void autortfm_did_free(void* ptr) AUTORTFM_NOEXCEPT;
+	UE_AUTORTFM_API void autortfm_did_free(void* ptr) AUTORTFM_NOEXCEPT;
 #else
 UE_AUTORTFM_CRITICAL_INLINE void autortfm_did_free(void* ptr) AUTORTFM_NOEXCEPT
 {
@@ -462,7 +481,7 @@ UE_AUTORTFM_CRITICAL_INLINE void autortfm_did_free(void* ptr) AUTORTFM_NOEXCEPT
 // be injected by the compiler into a global constructor in the AutoRTFM compiled
 // code.
 #if UE_AUTORTFM_ENABLED
-UE_AUTORTFM_API void autortfm_check_abi(void* ptr, size_t size) AUTORTFM_NOEXCEPT;
+	UE_AUTORTFM_API void autortfm_check_abi(void* ptr, size_t size) AUTORTFM_NOEXCEPT;
 #else
 UE_AUTORTFM_CRITICAL_INLINE void autortfm_check_abi(void* ptr, size_t size) AUTORTFM_NOEXCEPT
 {
@@ -472,10 +491,10 @@ UE_AUTORTFM_CRITICAL_INLINE void autortfm_check_abi(void* ptr, size_t size) AUTO
 #endif
 
 #if UE_AUTORTFM_ENABLED
-// Called when execution unexpectedly reaches a code path that was considered unreachable.
-// Either aborts execution of the program or aborts the current transaction, depending on the
-// current 'InternalAbortAction' state.
-[[noreturn]] UE_AUTORTFM_API void autortfm_unreachable(const char* Message) AUTORTFM_NOEXCEPT;
+	// Called when execution unexpectedly reaches a code path that was considered unreachable.
+	// Either aborts execution of the program or aborts the current transaction, depending on the
+	// current 'InternalAbortAction' state.
+	[[noreturn]] UE_AUTORTFM_API void autortfm_unreachable(const char* Message) AUTORTFM_NOEXCEPT;
 #else
 UE_AUTORTFM_CRITICAL_INLINE void autortfm_unreachable(const char* Message) AUTORTFM_NOEXCEPT
 {
@@ -484,8 +503,9 @@ UE_AUTORTFM_CRITICAL_INLINE void autortfm_unreachable(const char* Message) AUTOR
 #endif
 
 #if UE_AUTORTFM_ENABLED
-// Given an open function pointer, returns the equivalent closed function pointer.
-extern "C" UE_AUTORTFM_FORCENOINLINE UE_AUTORTFM_API void* autortfm_lookup_function(void* OriginalFunction, const char* Where) AUTORTFM_NOEXCEPT;
+	// Given an open function pointer, returns the equivalent closed function pointer.
+	extern "C" UE_AUTORTFM_FORCENOINLINE UE_AUTORTFM_API void* autortfm_lookup_function(
+		void* OriginalFunction, const char* Where) AUTORTFM_NOEXCEPT;
 #else
 UE_AUTORTFM_CRITICAL_INLINE void* autortfm_lookup_function(void* OriginalFunction, const char* Where) AUTORTFM_NOEXCEPT
 {
@@ -496,5 +516,5 @@ UE_AUTORTFM_CRITICAL_INLINE void* autortfm_lookup_function(void* OriginalFunctio
 #endif
 
 #ifdef __cplusplus
-} // extern "C"
+}  // extern "C"
 #endif
diff --git a/Engine/Source/Runtime/AutoRTFM/Public/AutoRTFM/Defines.h b/Engine/Source/Runtime/AutoRTFM/Public/AutoRTFM/Defines.h
--- a/Engine/Source/Runtime/AutoRTFM/Public/AutoRTFM/Defines.h
+++ b/Engine/Source/Runtime/AutoRTFM/Public/AutoRTFM/Defines.h
@@ -82,7 +82,7 @@
 // whether AutoRTFM instrumentation should be enabled for each individual
 // function based on the AutoRTFM-enabled state of each callee made by the
 // function. If the function calls any AutoRTFM-disabled function, then the
-// function will also be AutoRTFM-disabled, otherwise the function is 
+// function will also be AutoRTFM-disabled, otherwise the function is
 // AutoRTFM-enabled.
 // Applying the annotation to a class is equivalent to applying the annotation
 // to each method of the class.
@@ -135,7 +135,7 @@
 #define AUTORTFM_SANITIZER_ONLY(...)
 #endif
 
-#else // ^^^ UE_AUTORTFM ^^^ | vvv !UE_AUTORTFM vvv
+#else  // ^^^ UE_AUTORTFM ^^^ | vvv !UE_AUTORTFM vvv
 
 #define UE_AUTORTFM_NOAUTORTFM
 #define UE_AUTORTFM_ALWAYS_OPEN
@@ -176,7 +176,7 @@
 #define UE_AUTORTFM_ASSUME_SAFE
 #endif
 
-#define UE_AUTORTFM_CONCAT_IMPL(A, B) A ## B
+#define UE_AUTORTFM_CONCAT_IMPL(A, B) A##B
 #define UE_AUTORTFM_CONCAT(A, B) UE_AUTORTFM_CONCAT_IMPL(A, B)
 
 #if defined(_MSC_VER)
@@ -196,7 +196,6 @@
 #define UE_AUTORTFM_ASSUME(x)
 #endif
 
-
 #if (defined(UE_BUILD_DEBUG) && UE_BUILD_DEBUG) || (defined(AUTORTFM_BUILD_DEBUG) && AUTORTFM_BUILD_DEBUG)
 // Force-inlining can make debugging glitchy. Disable this if we're running a debug build.
 #undef UE_AUTORTFM_CALLSITE_FORCEINLINE
@@ -206,11 +205,8 @@
 #endif
 
 #ifdef _MSC_VER
-#define AUTORTFM_DISABLE_UNREACHABLE_CODE_WARNINGS \
-	__pragma (warning(push)) \
-	__pragma (warning(disable: 4702)) /* unreachable code */
-#define AUTORTFM_RESTORE_UNREACHABLE_CODE_WARNINGS \
-	__pragma (warning(pop))
+#define AUTORTFM_DISABLE_UNREACHABLE_CODE_WARNINGS __pragma(warning(push)) __pragma(warning(disable : 4702)) /* unreachable code */
+#define AUTORTFM_RESTORE_UNREACHABLE_CODE_WARNINGS __pragma(warning(pop))
 #else
 #define AUTORTFM_DISABLE_UNREACHABLE_CODE_WARNINGS
 #define AUTORTFM_RESTORE_UNREACHABLE_CODE_WARNINGS
@@ -224,6 +220,8 @@
 #define UE_AUTORTFM_CRITICAL_INLINE static UE_AUTORTFM_FORCEINLINE
 #define UE_AUTORTFM_CRITICAL_INLINE_ALWAYS static UE_AUTORTFM_FORCEINLINE_ALWAYS
 #else
-#define UE_AUTORTFM_CRITICAL_INLINE UE_AUTORTFM_FORCEINLINE // TODO: This needs to be revisited. we don't want bloated executables when modules are enabled
-#define UE_AUTORTFM_CRITICAL_INLINE_ALWAYS UE_AUTORTFM_FORCEINLINE_ALWAYS // TODO: This needs to be revisited. we don't want bloated executables when modules are enabled
+#define UE_AUTORTFM_CRITICAL_INLINE \
+	UE_AUTORTFM_FORCEINLINE  // TODO: This needs to be revisited. we don't want bloated executables when modules are enabled
+#define UE_AUTORTFM_CRITICAL_INLINE_ALWAYS \
+	UE_AUTORTFM_FORCEINLINE_ALWAYS  // TODO: This needs to be revisited. we don't want bloated executables when modules are enabled
 #endif
diff --git a/Engine/Source/Runtime/AutoRTFM/Public/AutoRTFM/ReportHazard.h b/Engine/Source/Runtime/AutoRTFM/Public/AutoRTFM/ReportHazard.h
--- a/Engine/Source/Runtime/AutoRTFM/Public/AutoRTFM/ReportHazard.h
+++ b/Engine/Source/Runtime/AutoRTFM/Public/AutoRTFM/ReportHazard.h
@@ -32,14 +32,13 @@ AUTORTFM_OPEN UE_AUTORTFM_API void ReportAutoRTFMHazard(EHazardType HazardType);
  * Reports an AutoRTFM hazard if this block of code is reached from the closed.
  * `autortfm_is_closed` is optimized away at compile time, so this compiles away to nothing in an open function.
  */
-#define UE_AUTORTFM_REPORT_HAZARD_IF_CLOSED(InHazardType)                                                       \
-    (autortfm_is_closed()                                                                                       \
-        ? ::AutoRTFM::ForTheRuntime::ReportAutoRTFMHazard(::AutoRTFM::ForTheRuntime::EHazardType::InHazardType) \
-        : (void)0)
+#define UE_AUTORTFM_REPORT_HAZARD_IF_CLOSED(InHazardType)                                                                         \
+	(autortfm_is_closed() ? ::AutoRTFM::ForTheRuntime::ReportAutoRTFMHazard(::AutoRTFM::ForTheRuntime::EHazardType::InHazardType) \
+						  : (void)0)
 
 #else
 
 /** Our macro compiles away to nothing on compilers which do not support AutoRTFM. */
 #define UE_AUTORTFM_REPORT_HAZARD_IF_CLOSED(InHazardType) ((void)0)
 
-#endif // UE_AUTORTFM
\ No newline at end of file
+#endif  // UE_AUTORTFM
\ No newline at end of file
diff --git a/Engine/Source/Runtime/AutoRTFM/Public/AutoRTFM/Testing.h b/Engine/Source/Runtime/AutoRTFM/Public/AutoRTFM/Testing.h
--- a/Engine/Source/Runtime/AutoRTFM/Public/AutoRTFM/Testing.h
+++ b/Engine/Source/Runtime/AutoRTFM/Public/AutoRTFM/Testing.h
@@ -10,13 +10,19 @@ namespace AutoRTFM::Testing
 // Implemented by the test framework, and called when AUTORTFM_TESTING_ASSERT() fails.
 void AssertionFailure(const char* Expression, const char* File, int Line);
 
-#define AUTORTFM_TESTING_ASSERT(Condition) \
-	do { if (!(Condition)) { AssertionFailure("AUTORTFM_TESTING_ASSERT(" #Condition ")", __FILE__, __LINE__); }} while (false)
+#define AUTORTFM_TESTING_ASSERT(Condition)                                                   \
+	do                                                                                       \
+	{                                                                                        \
+		if (!(Condition))                                                                    \
+		{                                                                                    \
+			AssertionFailure("AUTORTFM_TESTING_ASSERT(" #Condition ")", __FILE__, __LINE__); \
+		}                                                                                    \
+	} while (false)
 
 // Run the callback in a transaction like Transact, but abort program
 // execution if the result is anything other than autortfm_committed.
 // Useful for testing.
-template<typename TFunctor>
+template <typename TFunctor>
 static UE_AUTORTFM_FORCEINLINE void Commit(AUTORTFM_IMPLICIT_ENABLE const TFunctor& Functor)
 {
 	const ETransactionResult Result = Transact(Functor);
@@ -25,7 +31,7 @@ static UE_AUTORTFM_FORCEINLINE void Commit(AUTORTFM_IMPLICIT_ENABLE const TFunct
 
 // Run the callback in a transaction like Transact, but abort program
 // execution if the result is anything other than abort.
-template<typename TFunctor>
+template <typename TFunctor>
 static UE_AUTORTFM_FORCEINLINE void Abort(AUTORTFM_IMPLICIT_ENABLE const TFunctor& Functor)
 {
 	const ETransactionResult Result = Transact(Functor);
@@ -38,7 +44,10 @@ UE_AUTORTFM_API ForTheRuntime::EAutoRTFMEnabledState ForceSetAutoRTFMRuntime(For
 struct FEnabledStateResetterScoped final
 {
 	FEnabledStateResetterScoped(ForTheRuntime::EAutoRTFMEnabledState State) : Original(ForceSetAutoRTFMRuntime(State)) {}
-	~FEnabledStateResetterScoped() { ForceSetAutoRTFMRuntime(Original); }
+	~FEnabledStateResetterScoped()
+	{
+		ForceSetAutoRTFMRuntime(Original);
+	}
 
 private:
 	const ForTheRuntime::EAutoRTFMEnabledState Original;
@@ -46,4 +55,4 @@ private:
 
 #undef AUTORTFM_TESTING_ASSERT
 
-} // namespace AutoRTFM::Testing
+}  // namespace AutoRTFM::Testing
