[AutoRTFM] Reverse CallstackTree DAG

The callstacks returned by UE's stack walking utilities begin with the deepest frame, so the DAG is actually around the wrong way.
Fixing this reduces the amount of memory used by the CallstackTree.

#Synced-CL 49274188

diff --git a/Engine/Source/Programs/AutoRTFMTests/Private/CallstackTreeTests.cpp b/Engine/Source/Programs/AutoRTFMTests/Private/CallstackTreeTests.cpp
--- a/Engine/Source/Programs/AutoRTFMTests/Private/CallstackTreeTests.cpp
+++ b/Engine/Source/Programs/AutoRTFMTests/Private/CallstackTreeTests.cpp
@@ -58,23 +58,23 @@ TEST_CASE("CallstackTree")
 	SECTION("Deep")
 	{
 		AutoRTFM::FCallstackTree Tree;
-		AutoRTFM::TStack<autortfm_callstack_frame, 5> A{1, 2, 3, 4, 5};
+		AutoRTFM::TStack<autortfm_callstack_frame, 5> A{5, 4, 3, 2, 1};
 		FHandle HandleA = Tree.Add(A.Num(), &A.Front());
 		REQUIRE(HandleA == 5);
 
-		AutoRTFM::TStack<autortfm_callstack_frame, 5> B{1, 2, 3, 4, 5};
+		AutoRTFM::TStack<autortfm_callstack_frame, 5> B{5, 4, 3, 2, 1};
 		FHandle HandleB = Tree.Add(B.Num(), &B.Front());
 		REQUIRE(HandleB == 5);
 
-		AutoRTFM::TStack<autortfm_callstack_frame, 5> C{1, 2, 3};
+		AutoRTFM::TStack<autortfm_callstack_frame, 5> C{3, 2, 1};
 		FHandle HandleC = Tree.Add(C.Num(), &C.Front());
 		REQUIRE(HandleC == 3);
 
-		AutoRTFM::TStack<autortfm_callstack_frame, 5> D{4, 5};
+		AutoRTFM::TStack<autortfm_callstack_frame, 5> D{5, 4};
 		FHandle HandleD = Tree.Add(D.Num(), &D.Front());
 		REQUIRE(HandleD == 7);
 
-		AutoRTFM::TStack<autortfm_callstack_frame, 5> E{1, 2};
+		AutoRTFM::TStack<autortfm_callstack_frame, 5> E{2, 1};
 		FHandle HandleE = Tree.Add(E.Num(), &E.Front());
 		REQUIRE(HandleE == 2);
 
diff --git a/Engine/Source/Runtime/AutoRTFM/Private/CallstackTree.h b/Engine/Source/Runtime/AutoRTFM/Private/CallstackTree.h
--- a/Engine/Source/Runtime/AutoRTFM/Private/CallstackTree.h
+++ b/Engine/Source/Runtime/AutoRTFM/Private/CallstackTree.h
@@ -10,27 +10,32 @@
 #include "StlAllocator.h"
 #include "Utils.h"
 
-#include <algorithm>
 #include <cstdint>
 #include <unordered_set>
 
 namespace AutoRTFM
 {
 
+// FCallstackTree is a store for callstacks that internally holds a DAG of
+// frames which reduces memory consumption by de-duplicating common base frames.
 class AUTORTFM_DISABLE_NO_SANITIZE FCallstackTree final
 {
 public:
 	using FHandle = uint32_t;
 	static constexpr FHandle InvalidHandle = 0;
-	static constexpr FHandle MaxHandle = ~static_cast<FHandle>(0);
+	static constexpr FHandle MaxHandle = ~InvalidHandle;
 
-	FHandle Add(size_t StackCount, autortfm_callstack_frame const* StackTrace)
+	// Adds a callstack consisting of NumFrames frames at Frames.
+	// Frames starts with the most deeply-nested frame pointer.
+	// Returns a handle that can be passed to Get() to retrieve the callstack.
+	FHandle Add(size_t NumFrames, autortfm_callstack_frame const* Frames)
 	{
 		FHandle Parent = InvalidHandle;
 		FHandle Handle = InvalidHandle;
-		for (size_t I = 0; I < StackCount; I++)
+		size_t I = NumFrames;
+		while (I--)
 		{
-			FEntry const Entry{Parent, StackTrace[I]};
+			FEntry const Entry{Parent, Frames[I]};
 			if (auto It = Lookup.find(Entry); It != Lookup.end())
 			{
 				Handle = *It;
@@ -48,16 +53,16 @@ public:
 		return Handle;
 	}
 
+	// Retrieves a callstack previously added with Add().
 	TStack<autortfm_callstack_frame, 8> Get(FHandle Handle)
 	{
 		TStack<autortfm_callstack_frame, 8> StackTrace;
 		while (Handle != InvalidHandle)
 		{
 			FEntry const& Entry = Entries[Handle - 1];
-			StackTrace.Push(Entry.PC);
+			StackTrace.Push(Entry.Frame);
 			Handle = Entry.Parent;
 		}
-		std::reverse(StackTrace.begin(), StackTrace.end());
 		return StackTrace;
 	}
 
@@ -65,7 +70,7 @@ private:
 	struct FEntry
 	{
 		FHandle Parent;
-		autortfm_callstack_frame PC;
+		autortfm_callstack_frame Frame;
 
 		bool operator==(const FEntry&) const = default;
 	};
@@ -78,7 +83,7 @@ private:
 		const FEntryList& Entries;
 		size_t operator()(const FEntry& Entry) const
 		{
-			return (Entry.Parent * 31) ^ Entry.PC;
+			return (Entry.Parent * 31) ^ Entry.Frame;
 		}
 		size_t operator()(FHandle Handle) const
 		{
diff --git a/Engine/Source/Runtime/AutoRTFM/Public/AutoRTFM/CAPI.h b/Engine/Source/Runtime/AutoRTFM/Public/AutoRTFM/CAPI.h
--- a/Engine/Source/Runtime/AutoRTFM/Public/AutoRTFM/CAPI.h
+++ b/Engine/Source/Runtime/AutoRTFM/Public/AutoRTFM/CAPI.h
@@ -139,11 +139,11 @@ extern "C"
 		// Must not be null.
 		bool (*IsLogActive)(autortfm_log_severity Severity);
 
-		// Performs a capture of the currently executing callstack, writing at most
-		// MaxCount frame addresses to StackOut.
+		// Performs a capture of the currently executing callstack, writing at most MaxFrames frame
+		// addresses to StackOut, starting with with the most deeply nested frame.
 		// Returns the number of frame addresses actually written.
 		// Must not be null.
-		size_t (*CaptureCallstack)(size_t MaxCount, autortfm_callstack_frame* StackOut);
+		size_t (*CaptureCallstack)(size_t MaxFrames, autortfm_callstack_frame* StackOut);
 
 		// Prints a callstack captured with CaptureCallstack to the log with the
 		// given severity.
