autortfm compiler changes for 'redirection instrumentation'

- new cmdline flag 'autortfm-enable-redirection-instrumentation'
- when this flag is enabled, we instrument all loads and stores looking for operations to 'negative' pointers. for loads/stores to negative pointers, we redirect the operation to a runtime helper function to carry it out
- this is intended to help support the 'remote heap' feature of dstm

#Synced-CL 49706563

diff --git a/Engine/Source/Programs/UnrealBuildTool/Configuration/Rules/ReadOnlyTargetRules.cs b/Engine/Source/Programs/UnrealBuildTool/Configuration/Rules/ReadOnlyTargetRules.cs
--- a/Engine/Source/Programs/UnrealBuildTool/Configuration/Rules/ReadOnlyTargetRules.cs
+++ b/Engine/Source/Programs/UnrealBuildTool/Configuration/Rules/ReadOnlyTargetRules.cs
@@ -138,6 +138,7 @@ namespace UnrealBuildTool
 		public bool bUseAutoRTFMVerifier => Inner.bUseAutoRTFMVerifier;
 		public bool bAutoRTFMVerify => Inner.bAutoRTFMVerify;
 		public bool bAutoRTFMClosedStaticLinkage => Inner.bAutoRTFMClosedStaticLinkage;
+		public bool bAutoRTFMRedirectionInstrumentation => Inner.bAutoRTFMRedirectionInstrumentation;
 
 		public bool bUseChaosMemoryTracking => Inner.bUseChaosMemoryTracking;
 
diff --git a/Engine/Source/Programs/UnrealBuildTool/Configuration/Rules/TargetRules.cs b/Engine/Source/Programs/UnrealBuildTool/Configuration/Rules/TargetRules.cs
--- a/Engine/Source/Programs/UnrealBuildTool/Configuration/Rules/TargetRules.cs
+++ b/Engine/Source/Programs/UnrealBuildTool/Configuration/Rules/TargetRules.cs
@@ -990,6 +990,14 @@ namespace UnrealBuildTool
 		[XmlConfigFile(Category = "BuildConfiguration")]
 		public bool bAutoRTFMClosedStaticLinkage { get; set; } = false;
 
+		/// <summary>
+		/// Whether to instrument all memory accesses for redirection.
+		/// </summary>
+		[RequiresUniqueBuildEnvironment]
+		[CommandLine("-AutoRTFMRedirectionInstrumentation", Value = "true")]
+		[XmlConfigFile(Category = "BuildConfiguration")]
+		public bool bAutoRTFMRedirectionInstrumentation { get; set; } = false;
+
 		/// <summary>
 		/// Whether to compile in checked chaos features for debugging
 		/// </summary>
diff --git a/Engine/Source/Programs/UnrealBuildTool/Configuration/UEBuildModuleCPP.cs b/Engine/Source/Programs/UnrealBuildTool/Configuration/UEBuildModuleCPP.cs
--- a/Engine/Source/Programs/UnrealBuildTool/Configuration/UEBuildModuleCPP.cs
+++ b/Engine/Source/Programs/UnrealBuildTool/Configuration/UEBuildModuleCPP.cs
@@ -1972,6 +1972,7 @@ namespace UnrealBuildTool
 			CompileEnvironment.bEnableAutoRTFMVerification = ModuleCompileEnvironment.bEnableAutoRTFMVerification;
 			CompileEnvironment.bAutoRTFMVerify = ModuleCompileEnvironment.bAutoRTFMVerify;
 			CompileEnvironment.bAutoRTFMClosedStaticLinkage = ModuleCompileEnvironment.bAutoRTFMClosedStaticLinkage;
+			CompileEnvironment.bAutoRTFMRedirectionInstrumentation = ModuleCompileEnvironment.bAutoRTFMRedirectionInstrumentation;
 			CompileEnvironment.bValidateFormatStrings = ModuleCompileEnvironment.bValidateFormatStrings;
 			CompileEnvironment.bValidateInternalApi = ModuleCompileEnvironment.bValidateInternalApi;
 			CompileEnvironment.bValidateExperimentalApi = ModuleCompileEnvironment.bValidateExperimentalApi;
@@ -2626,6 +2627,7 @@ namespace UnrealBuildTool
 				Result.bEnableAutoRTFMVerification = Target.bUseAutoRTFMVerifier;
 				Result.bAutoRTFMVerify = Target.bAutoRTFMVerify;
 				Result.bAutoRTFMClosedStaticLinkage = Target.bAutoRTFMClosedStaticLinkage;
+				Result.bAutoRTFMRedirectionInstrumentation = Target.bAutoRTFMRedirectionInstrumentation;
 			}
 
 			if (Result.OptimizationLevel != Rules.OptimizationLevel)
diff --git a/Engine/Source/Programs/UnrealBuildTool/Platform/Clang/ClangToolChain.cs b/Engine/Source/Programs/UnrealBuildTool/Platform/Clang/ClangToolChain.cs
--- a/Engine/Source/Programs/UnrealBuildTool/Platform/Clang/ClangToolChain.cs
+++ b/Engine/Source/Programs/UnrealBuildTool/Platform/Clang/ClangToolChain.cs
@@ -464,6 +464,11 @@ namespace UnrealBuildTool
 			{
 				Arguments.Add("-Xclang -mllvm -Xclang -autortfm-closed-static-linkage");
 			}
+
+			if (CompileEnvironment.bAutoRTFMRedirectionInstrumentation)
+			{
+				Arguments.Add("-Xclang -mllvm -Xclang -autortfm-enable-redirection-instrumentation");
+			}
 		}
 
 		protected virtual void GetCStandardCompileArgument(CppCompileEnvironment CompileEnvironment, List<string> Arguments)
diff --git a/Engine/Source/Programs/UnrealBuildTool/Platform/Windows/VCToolChain.cs b/Engine/Source/Programs/UnrealBuildTool/Platform/Windows/VCToolChain.cs
--- a/Engine/Source/Programs/UnrealBuildTool/Platform/Windows/VCToolChain.cs
+++ b/Engine/Source/Programs/UnrealBuildTool/Platform/Windows/VCToolChain.cs
@@ -1539,6 +1539,11 @@ namespace UnrealBuildTool
 				{
 					Arguments.Add("-Xclang -mllvm -Xclang -autortfm-closed-static-linkage");
 				}
+
+				if (CompileEnvironment.bAutoRTFMRedirectionInstrumentation)
+				{
+					Arguments.Add("-Xclang -mllvm -Xclang -autortfm-enable-redirection-instrumentation");
+				}
 			}
 		}
 
@@ -1593,6 +1598,11 @@ namespace UnrealBuildTool
 				{
 					Arguments.Add("-Xclang -mllvm -Xclang -autortfm-closed-static-linkage");
 				}
+
+				if (CompileEnvironment.bAutoRTFMRedirectionInstrumentation)
+				{
+					Arguments.Add("-Xclang -mllvm -Xclang -autortfm-enable-redirection-instrumentation");
+				}
 			}
 		}
 
diff --git a/Engine/Source/Programs/UnrealBuildTool/System/CppCompileEnvironment.cs b/Engine/Source/Programs/UnrealBuildTool/System/CppCompileEnvironment.cs
--- a/Engine/Source/Programs/UnrealBuildTool/System/CppCompileEnvironment.cs
+++ b/Engine/Source/Programs/UnrealBuildTool/System/CppCompileEnvironment.cs
@@ -367,6 +367,11 @@ namespace UnrealBuildTool
 		/// </summary>
 		public bool bAutoRTFMClosedStaticLinkage = false;
 
+		/// <summary>
+		/// Whether we should instrument all memory accesses for redirection
+		/// </summary>
+		public bool bAutoRTFMRedirectionInstrumentation = false;
+
 		/// <summary>
 		/// If unity builds are enabled this can be used to override if this specific module will build using Unity.
 		/// This is set using the per module configurations in BuildConfiguration.
@@ -890,6 +895,7 @@ namespace UnrealBuildTool
 			AutoRTFMExternalMappingFiles = new(Other.AutoRTFMExternalMappingFiles);
 			bAutoRTFMVerify = Other.bAutoRTFMVerify;
 			bAutoRTFMClosedStaticLinkage = Other.bAutoRTFMClosedStaticLinkage;
+			bAutoRTFMRedirectionInstrumentation = Other.bAutoRTFMRedirectionInstrumentation;
 			bCodeCoverage = Other.bCodeCoverage;
 			OptimizationLevel = Other.OptimizationLevel;
 			FPSemantics = Other.FPSemantics;
diff --git a/Engine/Source/Runtime/AutoRTFM/Private/API.cpp b/Engine/Source/Runtime/AutoRTFM/Private/API.cpp
--- a/Engine/Source/Runtime/AutoRTFM/Private/API.cpp
+++ b/Engine/Source/Runtime/AutoRTFM/Private/API.cpp
@@ -40,10 +40,33 @@ bool GAutoRTFMEnsureOnInternalAbort = true;
 // ProcessAllPendingOpenToClosedRegistrations() when autortfm_initialize()
 // is called.
 autortfm_open_to_closed_table* GPendingOpenToClosedRegistrations = nullptr;
+}
 
 constexpr int GMaxRegisteredHeapRedirects = 4;
 AutoRTFM::FHeapRedirectCallbacks GRegisteredHeapRedirects[GMaxRegisteredHeapRedirects];
 int GNumRegisteredHeapRedirects = 0;
+
+AutoRTFM::FHeapRedirectCallbacks* FindHeapRedirectCallbacksForAddressSpace(uint32_t AddressSpace)
+{
+	static AutoRTFM::FHeapRedirectCallbacks* CachedRedirectCallbacks = nullptr;
+
+	if (CachedRedirectCallbacks && (CachedRedirectCallbacks->AddressSpace == AddressSpace))
+	{
+		return CachedRedirectCallbacks;
+	}
+
+	for (int RedirectedHeapIndex = 0; RedirectedHeapIndex < GNumRegisteredHeapRedirects; RedirectedHeapIndex++)
+	{
+		AutoRTFM::FHeapRedirectCallbacks* HeapRedirects = &GRegisteredHeapRedirects[RedirectedHeapIndex];
+
+		if (HeapRedirects->AddressSpace == AddressSpace)
+		{
+			CachedRedirectCallbacks = HeapRedirects;
+			return HeapRedirects;
+		}
+	}
+
+	return nullptr;
 }
 #endif  // UE_AUTORTFM
 
@@ -1127,29 +1150,81 @@ void ForTheRuntime::RegisterOnAbortFromTheOpen(TTask<void()>&& Work)
 
 void ForTheRuntime::RedirectedLoad(uint32_t AddressSpace, void* DestPointer, uint64_t Size, uint64_t SourceAddress)
 {
-	for (int RedirectedHeapIndex = 0; RedirectedHeapIndex < GNumRegisteredHeapRedirects; RedirectedHeapIndex++)
+	if (FHeapRedirectCallbacks* HeapRedirects = FindHeapRedirectCallbacksForAddressSpace(AddressSpace))
 	{
-		FHeapRedirectCallbacks* HeapRedirects = &GRegisteredHeapRedirects[RedirectedHeapIndex];
+		HeapRedirects->RedirectedLoad(DestPointer, Size, SourceAddress);
+	}
+}
 
-		if (HeapRedirects->AddressSpace == AddressSpace)
-		{
-			HeapRedirects->RedirectedLoad(DestPointer, Size, SourceAddress);
-			return;
-		}
+void ForTheRuntime::RedirectedLoad8(uint32_t AddressSpace, void* DestPointer, uint64_t SourceAddress)
+{
+	if (FHeapRedirectCallbacks* HeapRedirects = FindHeapRedirectCallbacksForAddressSpace(AddressSpace))
+	{
+		HeapRedirects->RedirectedLoad8(DestPointer, SourceAddress);
+	}
+}
+
+void ForTheRuntime::RedirectedLoad4(uint32_t AddressSpace, void* DestPointer, uint64_t SourceAddress)
+{
+	if (FHeapRedirectCallbacks* HeapRedirects = FindHeapRedirectCallbacksForAddressSpace(AddressSpace))
+	{
+		HeapRedirects->RedirectedLoad4(DestPointer, SourceAddress);
+	}
+}
+
+void ForTheRuntime::RedirectedLoad2(uint32_t AddressSpace, void* DestPointer, uint64_t SourceAddress)
+{
+	if (FHeapRedirectCallbacks* HeapRedirects = FindHeapRedirectCallbacksForAddressSpace(AddressSpace))
+	{
+		HeapRedirects->RedirectedLoad2(DestPointer, SourceAddress);
+	}
+}
+
+void ForTheRuntime::RedirectedLoad1(uint32_t AddressSpace, void* DestPointer, uint64_t SourceAddress)
+{
+	if (FHeapRedirectCallbacks* HeapRedirects = FindHeapRedirectCallbacksForAddressSpace(AddressSpace))
+	{
+		HeapRedirects->RedirectedLoad1(DestPointer, SourceAddress);
 	}
 }
 
 void ForTheRuntime::RedirectedStore(uint32_t AddressSpace, uint64_t DestAddress, uint64_t Size, const void* SourcePointer)
 {
-	for (int RedirectedHeapIndex = 0; RedirectedHeapIndex < GNumRegisteredHeapRedirects; RedirectedHeapIndex++)
+	if (FHeapRedirectCallbacks* HeapRedirects = FindHeapRedirectCallbacksForAddressSpace(AddressSpace))
+	{
+		HeapRedirects->RedirectedStore(DestAddress, Size, SourcePointer);
+	}
+}
+
+void ForTheRuntime::RedirectedStore8(uint32_t AddressSpace, uint64_t DestAddress, const void* SourcePointer)
+{
+	if (FHeapRedirectCallbacks* HeapRedirects = FindHeapRedirectCallbacksForAddressSpace(AddressSpace))
 	{
-		FHeapRedirectCallbacks* HeapRedirects = &GRegisteredHeapRedirects[RedirectedHeapIndex];
+		HeapRedirects->RedirectedStore8(DestAddress, SourcePointer);
+	}
+}
 
-		if (HeapRedirects->AddressSpace == AddressSpace)
-		{
-			HeapRedirects->RedirectedStore(DestAddress, Size, SourcePointer);
-			return;
-		}
+void ForTheRuntime::RedirectedStore4(uint32_t AddressSpace, uint64_t DestAddress, const void* SourcePointer)
+{
+	if (FHeapRedirectCallbacks* HeapRedirects = FindHeapRedirectCallbacksForAddressSpace(AddressSpace))
+	{
+		HeapRedirects->RedirectedStore4(DestAddress, SourcePointer);
+	}
+}
+
+void ForTheRuntime::RedirectedStore2(uint32_t AddressSpace, uint64_t DestAddress, const void* SourcePointer)
+{
+	if (FHeapRedirectCallbacks* HeapRedirects = FindHeapRedirectCallbacksForAddressSpace(AddressSpace))
+	{
+		HeapRedirects->RedirectedStore2(DestAddress, SourcePointer);
+	}
+}
+
+void ForTheRuntime::RedirectedStore1(uint32_t AddressSpace, uint64_t DestAddress, const void* SourcePointer)
+{
+	if (FHeapRedirectCallbacks* HeapRedirects = FindHeapRedirectCallbacksForAddressSpace(AddressSpace))
+	{
+		HeapRedirects->RedirectedStore1(DestAddress, SourcePointer);
 	}
 }
 
@@ -1167,15 +1242,29 @@ void RegisterHeapRedirectCallbacks(FHeapRedirectCallbacks Callbacks)
 {
 	AUTORTFM_ASSERT(GNumRegisteredHeapRedirects < GMaxRegisteredHeapRedirects);
 
+	// helper for handling casting the return type to the input function signature
+	auto LookupClosedHelper = [](auto InFunction) -> decltype(InFunction)
+	{
+		using FnType = decltype(InFunction);
+
+		return reinterpret_cast<FnType>(autortfm_lookup_function(reinterpret_cast<void*>(InFunction), "RegisterHeapRedirects"));
+	};
+
 	// when loads and stores get redirected, it seems useful to not make a closed->open transition
 	// so that the redirect callbacks can abort the transaction, etc. So here in the registration
 	// function we translate the function pointers passed in to be their closed variants
 	FHeapRedirectCallbacks ClosedCallbacks;
 	ClosedCallbacks.AddressSpace = Callbacks.AddressSpace;
-	ClosedCallbacks.RedirectedLoad = reinterpret_cast<void (*)(void*, uint64_t, uint64_t)>(
-		autortfm_lookup_function(reinterpret_cast<void*>(Callbacks.RedirectedLoad), "RegisterHeapRedirects"));
-	ClosedCallbacks.RedirectedStore = reinterpret_cast<void (*)(uint64_t, uint64_t, const void*)>(
-		autortfm_lookup_function(reinterpret_cast<void*>(Callbacks.RedirectedStore), "RegisterHeapRedirects"));
+	ClosedCallbacks.RedirectedLoad = LookupClosedHelper(Callbacks.RedirectedLoad);
+	ClosedCallbacks.RedirectedLoad8 = LookupClosedHelper(Callbacks.RedirectedLoad8);
+	ClosedCallbacks.RedirectedLoad4 = LookupClosedHelper(Callbacks.RedirectedLoad4);
+	ClosedCallbacks.RedirectedLoad2 = LookupClosedHelper(Callbacks.RedirectedLoad2);
+	ClosedCallbacks.RedirectedLoad1 = LookupClosedHelper(Callbacks.RedirectedLoad1);
+	ClosedCallbacks.RedirectedStore = LookupClosedHelper(Callbacks.RedirectedStore);
+	ClosedCallbacks.RedirectedStore8 = LookupClosedHelper(Callbacks.RedirectedStore8);
+	ClosedCallbacks.RedirectedStore4 = LookupClosedHelper(Callbacks.RedirectedStore4);
+	ClosedCallbacks.RedirectedStore2 = LookupClosedHelper(Callbacks.RedirectedStore2);
+	ClosedCallbacks.RedirectedStore1 = LookupClosedHelper(Callbacks.RedirectedStore1);
 	GRegisteredHeapRedirects[GNumRegisteredHeapRedirects++] = ClosedCallbacks;
 }
 
diff --git a/Engine/Source/Runtime/AutoRTFM/Private/CompilerInterface.cpp b/Engine/Source/Runtime/AutoRTFM/Private/CompilerInterface.cpp
--- a/Engine/Source/Runtime/AutoRTFM/Private/CompilerInterface.cpp
+++ b/Engine/Source/Runtime/AutoRTFM/Private/CompilerInterface.cpp
@@ -86,68 +86,98 @@ extern "C" AUTORTFM_DISABLE_NO_SANITIZE UE_AUTORTFM_API AUTORTFM_NO_ASAN void au
 	}
 }
 
+// there are two kinds of redirected loads/stores - ones that are explicitly marked up with an address
+// space, and others where the AddressSpace == 0 and the information is stored in the bottom nibble
+// of the top byte. This helper is used to extract the nibble
+static uint32 AutoRTFMExtractAddressSpaceFromAddress(uint64 Address)
+{
+	return (uint32)((Address >> 56) & 0x0F);
+}
+
+static bool AutoRTFMIsPointerNegative(const void* InPointer)
+{
+	union
+	{
+		int64 Address;
+		const void* Pointer;
+	};
+
+	Pointer = InPointer;
+	return Address < 0;
+}
+
 extern "C" UE_AUTORTFM_API AUTORTFM_NO_ASAN void autortfm_redirected_load(
 	uint32 AddressSpace, void* DestPointer, uint64 Size, uint64 SourceAddress) noexcept
 {
-	ForTheRuntime::RedirectedLoad(AddressSpace, DestPointer, Size, SourceAddress);
+	ForTheRuntime::RedirectedLoad(
+		AddressSpace ? AddressSpace : AutoRTFMExtractAddressSpaceFromAddress(SourceAddress), DestPointer, Size, SourceAddress);
 }
 
 extern "C" UE_AUTORTFM_API AUTORTFM_NO_ASAN uint64 autortfm_redirected_load_8(uint32 AddressSpace, uint64 SourceAddress) noexcept
 {
 	uint64 Result = 0;
-	ForTheRuntime::RedirectedLoad(AddressSpace, &Result, 8, SourceAddress);
+	ForTheRuntime::RedirectedLoad8(
+		AddressSpace ? AddressSpace : AutoRTFMExtractAddressSpaceFromAddress(SourceAddress), &Result, SourceAddress);
 	return Result;
 }
 
 extern "C" UE_AUTORTFM_API AUTORTFM_NO_ASAN uint32 autortfm_redirected_load_4(uint32 AddressSpace, uint64 SourceAddress) noexcept
 {
 	uint32 Result = 0;
-	ForTheRuntime::RedirectedLoad(AddressSpace, &Result, 4, SourceAddress);
+	ForTheRuntime::RedirectedLoad4(
+		AddressSpace ? AddressSpace : AutoRTFMExtractAddressSpaceFromAddress(SourceAddress), &Result, SourceAddress);
 	return Result;
 }
 
 extern "C" UE_AUTORTFM_API AUTORTFM_NO_ASAN uint16 autortfm_redirected_load_2(uint32 AddressSpace, uint64 SourceAddress) noexcept
 {
 	uint16 Result = 0;
-	ForTheRuntime::RedirectedLoad(AddressSpace, &Result, 2, SourceAddress);
+	ForTheRuntime::RedirectedLoad2(
+		AddressSpace ? AddressSpace : AutoRTFMExtractAddressSpaceFromAddress(SourceAddress), &Result, SourceAddress);
 	return Result;
 }
 
 extern "C" UE_AUTORTFM_API AUTORTFM_NO_ASAN uint8 autortfm_redirected_load_1(uint32 AddressSpace, uint64 SourceAddress) noexcept
 {
 	uint8 Result = 0;
-	ForTheRuntime::RedirectedLoad(AddressSpace, &Result, 1, SourceAddress);
+	ForTheRuntime::RedirectedLoad1(
+		AddressSpace ? AddressSpace : AutoRTFMExtractAddressSpaceFromAddress(SourceAddress), &Result, SourceAddress);
 	return Result;
 }
 
 extern "C" UE_AUTORTFM_API AUTORTFM_NO_ASAN void autortfm_redirected_store(
 	uint32 AddressSpace, uint64 DestAddress, uint64 Size, const void* SourcePointer) noexcept
 {
-	ForTheRuntime::RedirectedStore(AddressSpace, DestAddress, Size, SourcePointer);
+	ForTheRuntime::RedirectedStore(
+		AddressSpace ? AddressSpace : AutoRTFMExtractAddressSpaceFromAddress(DestAddress), DestAddress, Size, SourcePointer);
 }
 
 extern "C" UE_AUTORTFM_API AUTORTFM_NO_ASAN void autortfm_redirected_store_8(
 	uint32 AddressSpace, uint64 DestAddress, uint64 Value) noexcept
 {
-	ForTheRuntime::RedirectedStore(AddressSpace, DestAddress, sizeof(Value), &Value);
+	ForTheRuntime::RedirectedStore8(
+		AddressSpace ? AddressSpace : AutoRTFMExtractAddressSpaceFromAddress(DestAddress), DestAddress, &Value);
 }
 
 extern "C" UE_AUTORTFM_API AUTORTFM_NO_ASAN void autortfm_redirected_store_4(
 	uint32 AddressSpace, uint64 DestAddress, uint32 Value) noexcept
 {
-	ForTheRuntime::RedirectedStore(AddressSpace, DestAddress, sizeof(Value), &Value);
+	ForTheRuntime::RedirectedStore4(
+		AddressSpace ? AddressSpace : AutoRTFMExtractAddressSpaceFromAddress(DestAddress), DestAddress, &Value);
 }
 
 extern "C" UE_AUTORTFM_API AUTORTFM_NO_ASAN void autortfm_redirected_store_2(
 	uint32 AddressSpace, uint64 DestAddress, uint16 Value) noexcept
 {
-	ForTheRuntime::RedirectedStore(AddressSpace, DestAddress, sizeof(Value), &Value);
+	ForTheRuntime::RedirectedStore2(
+		AddressSpace ? AddressSpace : AutoRTFMExtractAddressSpaceFromAddress(DestAddress), DestAddress, &Value);
 }
 
 extern "C" UE_AUTORTFM_API AUTORTFM_NO_ASAN void autortfm_redirected_store_1(
 	uint32 AddressSpace, uint64 DestAddress, uint8 Value) noexcept
 {
-	ForTheRuntime::RedirectedStore(AddressSpace, DestAddress, sizeof(Value), &Value);
+	ForTheRuntime::RedirectedStore1(
+		AddressSpace ? AddressSpace : AutoRTFMExtractAddressSpaceFromAddress(DestAddress), DestAddress, &Value);
 }
 
 // Note: Internal.aem maps this to itself when called in the closed
@@ -174,6 +204,196 @@ extern "C" AUTORTFM_DISABLE_NO_SANITIZE UE_AUTORTFM_API void autortfm_memset(voi
 	Memset(Dst, Value, Size, Context);
 }
 
+template <bool bDstIsRedirected, bool bSrcIsRedirected>
+AUTORTFM_DISABLE_NO_SANITIZE void autortfm_memmove_redirectable_impl(void* Dst, const void* Src, size_t Size) noexcept
+{
+	int8 MoveBackwards = (Src < Dst);
+
+	uint8* DstCursor = (uint8*)Dst;
+	uint8* SrcCursor = (uint8*)Src;
+
+	if (MoveBackwards)
+	{
+		DstCursor += Size;
+		SrcCursor += Size;
+	}
+
+	while (Size >= 8)
+	{
+		uint64 Data = 0;
+
+		// if we are going backwards, the cursor starts out just beyond the end
+		if (MoveBackwards)
+		{
+			DstCursor -= 8;
+			SrcCursor -= 8;
+		}
+
+		if (bSrcIsRedirected)
+		{
+			Data = autortfm_redirected_load_8(0, (uint64)SrcCursor);
+		}
+		else
+		{
+			// use memcpy as this is potentially an unaligned load
+			memcpy(&Data, SrcCursor, 8);
+		}
+
+		if (bDstIsRedirected)
+		{
+			autortfm_redirected_store_8(0, (uint64)DstCursor, Data);
+		}
+		else
+		{
+			// use memcpy as this is potentially an unaligned store
+			memcpy(DstCursor, &Data, 8);
+		}
+
+		if (!MoveBackwards)
+		{
+			DstCursor += 8;
+			SrcCursor += 8;
+		}
+
+		Size -= 8;
+	}
+
+	while (Size > 0)
+	{
+		uint8 Data = 0;
+
+		if (MoveBackwards)
+		{
+			DstCursor -= 1;
+			SrcCursor -= 1;
+		}
+
+		if (bSrcIsRedirected)
+		{
+			Data = autortfm_redirected_load_1(0, (uint64)SrcCursor);
+		}
+		else
+		{
+			Data = *SrcCursor;
+		}
+
+		if (bDstIsRedirected)
+		{
+			autortfm_redirected_store_1(0, (uint64)DstCursor, Data);
+		}
+		else
+		{
+			*DstCursor = Data;
+		}
+
+		if (!MoveBackwards)
+		{
+			DstCursor += 1;
+			SrcCursor += 1;
+		}
+
+		Size -= 1;
+	}
+}
+
+template <bool bDstIsRedirected>
+AUTORTFM_DISABLE_NO_SANITIZE void autortfm_memset_redirectable_impl(void* Dst, uint8 Value, size_t Size) noexcept
+{
+	uint8* DstCursor = (uint8*)Dst;
+
+	uint64 Data8 = ((uint64)Value) * 0x0101010101010101ull;
+
+	while (Size >= 8)
+	{
+		if (bDstIsRedirected)
+		{
+			autortfm_redirected_store_8(0, (uint64)DstCursor, Data8);
+		}
+		else
+		{
+			// use memcpy as this is potentially an unaligned store
+			memcpy(DstCursor, &Data8, 8);
+		}
+
+		DstCursor += 8;
+		Size -= 8;
+	}
+
+	while (Size > 0)
+	{
+		uint8 Data = Value;
+
+		if (bDstIsRedirected)
+		{
+			autortfm_redirected_store_1(0, (uint64)DstCursor, Data);
+		}
+		else
+		{
+			*DstCursor = Data;
+		}
+
+		DstCursor += 1;
+		Size -= 1;
+	}
+}
+
+extern "C" AUTORTFM_DISABLE_NO_SANITIZE UE_AUTORTFM_API void autortfm_memmove_redirectable(
+	void* Dst, const void* Src, size_t Size) noexcept
+{
+	FContext* Context = FContext::Get();
+
+	bool bDstIsRedirected = AutoRTFMIsPointerNegative(Dst);
+	bool bSrcIsRedirected = AutoRTFMIsPointerNegative(Src);
+
+	if (bDstIsRedirected)
+	{
+		if (bSrcIsRedirected)
+		{
+			autortfm_memmove_redirectable_impl<true, true>(Dst, Src, Size);
+		}
+		else
+		{
+			autortfm_memmove_redirectable_impl<true, false>(Dst, Src, Size);
+		}
+	}
+	else
+	{
+		Context->RecordWrite(Dst, Size);
+
+		if (bSrcIsRedirected)
+		{
+			autortfm_memmove_redirectable_impl<false, true>(Dst, Src, Size);
+		}
+		else
+		{
+			autortfm_memmove_redirectable_impl<false, false>(Dst, Src, Size);
+		}
+	}
+}
+
+extern "C" AUTORTFM_DISABLE_NO_SANITIZE UE_AUTORTFM_API void autortfm_memcpy_redirectable(
+	void* Dst, const void* Src, size_t Size) noexcept
+{
+	autortfm_memmove_redirectable(Dst, Src, Size);
+}
+
+extern "C" AUTORTFM_DISABLE_NO_SANITIZE UE_AUTORTFM_API void autortfm_memset_redirectable(void* Dst, int Value, size_t Size) noexcept
+{
+	FContext* Context = FContext::Get();
+
+	bool bDstIsRedirected = AutoRTFMIsPointerNegative(Dst);
+
+	if (bDstIsRedirected)
+	{
+		autortfm_memset_redirectable_impl<true>(Dst, (uint8)Value, Size);
+	}
+	else
+	{
+		Context->RecordWrite(Dst, Size);
+		autortfm_memset_redirectable_impl<false>(Dst, (uint8)Value, Size);
+	}
+}
+
 extern "C" UE_AUTORTFM_ALWAYS_OPEN UE_AUTORTFM_API void autortfm_unreachable(const char* Message) noexcept
 {
 	AUTORTFM_REPORT_ERROR("AutoRTFM Unreachable: %s", Message);
diff --git a/Engine/Source/Runtime/AutoRTFM/Public/AutoRTFM.h b/Engine/Source/Runtime/AutoRTFM/Public/AutoRTFM.h
--- a/Engine/Source/Runtime/AutoRTFM/Public/AutoRTFM.h
+++ b/Engine/Source/Runtime/AutoRTFM/Public/AutoRTFM.h
@@ -245,7 +245,15 @@ UE_AUTORTFM_API void PopAllOnAbortHandlersInternal(const void* Key);
 UE_AUTORTFM_API void RegisterOnCommitFromTheOpen(TTask<void()>&& Work);
 UE_AUTORTFM_API void RegisterOnAbortFromTheOpen(TTask<void()>&& Work);
 UE_AUTORTFM_API void RedirectedLoad(uint32_t AddressSpace, void* DestPointer, uint64_t Size, uint64_t SourceAddress);
+UE_AUTORTFM_API void RedirectedLoad8(uint32_t AddressSpace, void* DestPointer, uint64_t SourceAddress);
+UE_AUTORTFM_API void RedirectedLoad4(uint32_t AddressSpace, void* DestPointer, uint64_t SourceAddress);
+UE_AUTORTFM_API void RedirectedLoad2(uint32_t AddressSpace, void* DestPointer, uint64_t SourceAddress);
+UE_AUTORTFM_API void RedirectedLoad1(uint32_t AddressSpace, void* DestPointer, uint64_t SourceAddress);
 UE_AUTORTFM_API void RedirectedStore(uint32_t AddressSpace, uint64_t DestAddress, uint64_t Size, const void* SourcePointer);
+UE_AUTORTFM_API void RedirectedStore8(uint32_t AddressSpace, uint64_t DestAddress, const void* SourcePointer);
+UE_AUTORTFM_API void RedirectedStore4(uint32_t AddressSpace, uint64_t DestAddress, const void* SourcePointer);
+UE_AUTORTFM_API void RedirectedStore2(uint32_t AddressSpace, uint64_t DestAddress, const void* SourcePointer);
+UE_AUTORTFM_API void RedirectedStore1(uint32_t AddressSpace, uint64_t DestAddress, const void* SourcePointer);
 }  // namespace ForTheRuntime
 #endif
 
@@ -795,7 +803,15 @@ struct FHeapRedirectCallbacks
 {
 	uint32_t AddressSpace;
 	void (*RedirectedLoad)(void* DestPointer, uint64_t Size, uint64_t SourceAddress);
+	void (*RedirectedLoad8)(void* DestPointer, uint64_t SourceAddress);
+	void (*RedirectedLoad4)(void* DestPointer, uint64_t SourceAddress);
+	void (*RedirectedLoad2)(void* DestPointer, uint64_t SourceAddress);
+	void (*RedirectedLoad1)(void* DestPointer, uint64_t SourceAddress);
 	void (*RedirectedStore)(uint64_t DestAddress, uint64_t Size, const void* SourcePointer);
+	void (*RedirectedStore8)(uint64_t DestAddress, const void* SourcePointer);
+	void (*RedirectedStore4)(uint64_t DestAddress, const void* SourcePointer);
+	void (*RedirectedStore2)(uint64_t DestAddress, const void* SourcePointer);
+	void (*RedirectedStore1)(uint64_t DestAddress, const void* SourcePointer);
 };
 
 #if UE_AUTORTFM
