[AutoRTFM] Fix non-X64 windows builds.

#Synced-CL 49880520

diff --git a/Engine/Source/Runtime/AutoRTFM/Private/TestUtils.cpp b/Engine/Source/Runtime/AutoRTFM/Private/TestUtils.cpp
--- a/Engine/Source/Runtime/AutoRTFM/Private/TestUtils.cpp
+++ b/Engine/Source/Runtime/AutoRTFM/Private/TestUtils.cpp
@@ -83,7 +83,10 @@ AUTORTFM_DISABLE static LONG WINAPI CrashHandler(_EXCEPTION_POINTERS* Info)
 		default:
 			return EXCEPTION_CONTINUE_SEARCH;
 	}
-	void* ProgramCounter = reinterpret_cast<void*>(Info->ContextRecord->Rip);
+	void* ProgramCounter = nullptr;
+#if AUTORTFM_ARCHITECTURE_X64
+	ProgramCounter = reinterpret_cast<void*>(Info->ContextRecord->Rip);
+#endif // TODO: ARM64.
 	GCrashHandlerCallback(FCrashInfo{.Kind = Kind, .ProgramCounter = ProgramCounter});
 	return EXCEPTION_CONTINUE_SEARCH;
 }
