[AutoRTFM] Migrate AUTORTFM_(UN)LIKELY into Defines.h.

This can help avoid a circular dependency when trying to use AUTORTFM_UNLIKELY in the memory stats header.

#Synced-CL 50109039

diff --git a/Engine/Source/Runtime/AutoRTFM/Private/API.cpp b/Engine/Source/Runtime/AutoRTFM/Private/API.cpp
--- a/Engine/Source/Runtime/AutoRTFM/Private/API.cpp
+++ b/Engine/Source/Runtime/AutoRTFM/Private/API.cpp
@@ -1,7 +1,6 @@
 // Copyright Epic Games, Inc. All Rights Reserved.
 
 #include "AutoRTFM.h"
-#include "AutoRTFM/Defines.h"
 #include "Context.h"
 #include "ExternAPI.h"
 #include "Transaction.h"
diff --git a/Engine/Source/Runtime/AutoRTFM/Private/Context.h b/Engine/Source/Runtime/AutoRTFM/Private/Context.h
--- a/Engine/Source/Runtime/AutoRTFM/Private/Context.h
+++ b/Engine/Source/Runtime/AutoRTFM/Private/Context.h
@@ -4,6 +4,7 @@
 
 #if (defined(__AUTORTFM) && __AUTORTFM)
 
+#include "AutoRTFM/Defines.h"
 #include "AutoRTFM/Task.h"
 #include "CallNest.h"
 #include "IntrusivePool.h"
diff --git a/Engine/Source/Runtime/AutoRTFM/Private/ContextInlines.h b/Engine/Source/Runtime/AutoRTFM/Private/ContextInlines.h
--- a/Engine/Source/Runtime/AutoRTFM/Private/ContextInlines.h
+++ b/Engine/Source/Runtime/AutoRTFM/Private/ContextInlines.h
@@ -4,6 +4,7 @@
 
 #if (defined(__AUTORTFM) && __AUTORTFM)
 
+#include "AutoRTFM/Defines.h"
 #include "Context.h"
 #include "TransactionInlines.h"
 
diff --git a/Engine/Source/Runtime/AutoRTFM/Private/TransactionInlines.h b/Engine/Source/Runtime/AutoRTFM/Private/TransactionInlines.h
--- a/Engine/Source/Runtime/AutoRTFM/Private/TransactionInlines.h
+++ b/Engine/Source/Runtime/AutoRTFM/Private/TransactionInlines.h
@@ -4,6 +4,7 @@
 
 #if (defined(__AUTORTFM) && __AUTORTFM)
 
+#include "AutoRTFM/Defines.h"
 #include "AutoRTFM/ReportHazard.h"
 #include "Context.h"
 #include "HitSet.h"
diff --git a/Engine/Source/Runtime/AutoRTFM/Private/Utils.h b/Engine/Source/Runtime/AutoRTFM/Private/Utils.h
--- a/Engine/Source/Runtime/AutoRTFM/Private/Utils.h
+++ b/Engine/Source/Runtime/AutoRTFM/Private/Utils.h
@@ -175,9 +175,6 @@ inline constexpr T Log2(T Value)
 
 }  // namespace AutoRTFM
 
-#define AUTORTFM_LIKELY(x) __builtin_expect(!!(x), 1)
-#define AUTORTFM_UNLIKELY(x) __builtin_expect(!!(x), 0)
-
 #define AUTORTFM_REQUIRE_SEMICOLON static_assert(true)
 
 #define AUTORTFM_REPORT_ERROR(Format, ...) \
diff --git a/Engine/Source/Runtime/AutoRTFM/Public/AutoRTFM/Defines.h b/Engine/Source/Runtime/AutoRTFM/Public/AutoRTFM/Defines.h
--- a/Engine/Source/Runtime/AutoRTFM/Public/AutoRTFM/Defines.h
+++ b/Engine/Source/Runtime/AutoRTFM/Public/AutoRTFM/Defines.h
@@ -129,6 +129,12 @@
 // Evaluates to true iff the AutoRTFM sanitizer is enabled.
 #define AUTORTFM_SANITIZER (__has_feature(autortfm_sanitizer))
 
+// Provides a hint to the compiler that the expression is likely to be true
+#define AUTORTFM_LIKELY(x) __builtin_expect(!!(x), 1)
+
+// Provides a hint to the compiler that the expression is unlikely to be true
+#define AUTORTFM_UNLIKELY(x) __builtin_expect(!!(x), 0)
+
 #if AUTORTFM_SANITIZER
 #define AUTORTFM_SANITIZER_ONLY(...) __VA_ARGS__
 #else
@@ -157,6 +163,8 @@
 #define AUTORTFM_MODE_OF_CALL(EXPR_OR_TYPE) autortfm_mode_disable
 #define UE_AUTORTFM_CALLSITE_FORCEINLINE
 #define AUTORTFM_SANITIZER (false)
+#define AUTORTFM_LIKELY(x) (!!(x))
+#define AUTORTFM_UNLIKELY(x) (!!(x))
 
 #endif
 
