Change up how RTFM ONABORT and ONPREABORT are compiled out when RTFM is not enabled to avoid unreachable code warnings.
#rnx

#Synced-CL 50073302

diff --git a/Engine/Source/Runtime/AutoRTFM/Public/AutoRTFM.h b/Engine/Source/Runtime/AutoRTFM/Public/AutoRTFM.h
--- a/Engine/Source/Runtime/AutoRTFM/Public/AutoRTFM.h
+++ b/Engine/Source/Runtime/AutoRTFM/Public/AutoRTFM.h
@@ -1223,8 +1223,8 @@ struct FThreadLocalHelper
 #define UE_AUTORTFM_OPEN_IMPL
 #define UE_AUTORTFM_OPEN_NO_SANITIZE_IMPL
 #define UE_AUTORTFM_CLOSE_IMPL
-#define UE_AUTORTFM_ONPREABORT_IMPL(...) while (false)
-#define UE_AUTORTFM_ONABORT_IMPL(...) while (false)
+#define UE_AUTORTFM_ONPREABORT_IMPL(...) (void) [__VA_ARGS__] () mutable
+#define UE_AUTORTFM_ONABORT_IMPL(...) (void) [__VA_ARGS__] () mutable
 #define UE_AUTORTFM_ONCOMMIT_IMPL(...)
 #define UE_AUTORTFM_TRANSACT_IMPL
 #endif
