[AutoRTFM] Factor Metrics code out of FContext.

This will allow metrics to be accessed from other places, like FTransaction.

#Synced-CL 50109039

diff --git a/Engine/Source/Runtime/AutoRTFM/Private/Context.cpp b/Engine/Source/Runtime/AutoRTFM/Private/Context.cpp
--- a/Engine/Source/Runtime/AutoRTFM/Private/Context.cpp
+++ b/Engine/Source/Runtime/AutoRTFM/Private/Context.cpp
@@ -6,10 +6,10 @@
 #include "ContextInlines.h"
 
 #include "AutoRTFM.h"
-#include "AutoRTFM/Metrics.h"
 #include "CallNestInlines.h"
 #include "ExternAPI.h"
 #include "FunctionMap.h"
+#include "MetricsPriv.h"
 #include "ScopedGuard.h"
 #include "StackRange.h"
 #include "Stats.h"
@@ -21,11 +21,6 @@
 #include "WindowsHeader.h"
 #endif
 
-namespace
-{
-AutoRTFM::FAutoRTFMMetrics GAutoRTFMMetrics;
-}
-
 namespace AutoRTFM
 {
 
@@ -55,17 +50,6 @@ void FContext::Destroy()
 	}
 }
 
-void ResetAutoRTFMMetrics()
-{
-	GAutoRTFMMetrics = FAutoRTFMMetrics{};
-}
-
-// get a snapshot of the current internal metrics
-FAutoRTFMMetrics GetAutoRTFMMetrics()
-{
-	return GAutoRTFMMetrics;
-}
-
 FContext::~FContext()
 {
 #if AUTORTFM_SANITIZER
@@ -125,7 +109,7 @@ void FContext::StartTransaction()
 	AUTORTFM_ENSURE_MSG(CurrentTransaction, "FContext::StartTransaction() can only be called within a scoped transaction");
 	AUTORTFM_ASSERT(Status == EContextStatus::OnTrack);
 
-	GAutoRTFMMetrics.NumTransactionsStarted++;
+	IncrementNumTransactionsStartedMetric();
 }
 
 void FContext::StartNonDeferredTransaction()
@@ -142,7 +126,7 @@ void FContext::StartNonDeferredTransaction()
 	// This form of transaction is always ultimately within a scoped Transact
 	AUTORTFM_ASSERT(Status == EContextStatus::OnTrack);
 
-	GAutoRTFMMetrics.NumTransactionsStarted++;
+	IncrementNumTransactionsStartedMetric();
 }
 
 ETransactionResult FContext::CommitTransaction()
@@ -187,14 +171,14 @@ ETransactionResult FContext::CommitTransaction()
 		PopTransaction();
 	}
 
-	GAutoRTFMMetrics.NumTransactionsCommitted++;
+	IncrementNumTransactionsCommittedMetric();
 
 	return Result;
 }
 
 void FContext::RollbackTransaction(EContextStatus NewStatus)
 {
-	GAutoRTFMMetrics.NumTransactionsAborted++;
+	IncrementNumTransactionsAbortedMetric();
 
 	AUTORTFM_ASSERT(Status == EContextStatus::OnTrack);
 	AUTORTFM_ASSERT(IsStatusAborting(NewStatus));
@@ -640,7 +624,7 @@ ETransactionResult FContext::Transact(void (*UninstrumentedFunction)(void*), voi
 void FContext::AbortByRequestAndThrow()
 {
 	AUTORTFM_ASSERT(Status == EContextStatus::OnTrack);
-	GAutoRTFMMetrics.NumTransactionsAbortedByRequest++;
+	IncrementNumTransactionsAbortedByRequestMetric();
 	RollbackTransaction(EContextStatus::AbortedByRequest);
 	Throw();
 }
@@ -648,7 +632,7 @@ void FContext::AbortByRequestAndThrow()
 void FContext::AbortByLanguageAndThrow()
 {
 	AUTORTFM_ASSERT(Status == EContextStatus::OnTrack);
-	GAutoRTFMMetrics.NumTransactionsAbortedByLanguage++;
+	IncrementNumTransactionsAbortedByLanguageMetric();
 	RollbackTransaction(EContextStatus::AbortedByLanguage);
 	Throw();
 }
diff --git a/Engine/Source/Runtime/AutoRTFM/Private/Metrics.cpp b/Engine/Source/Runtime/AutoRTFM/Private/Metrics.cpp
new file mode 100644
--- /dev/null
+++ b/Engine/Source/Runtime/AutoRTFM/Private/Metrics.cpp
@@ -0,0 +1,36 @@
+// Copyright Epic Games, Inc. All Rights Reserved.
+
+#include "AutoRTFM/Metrics.h"
+#include "MetricsPriv.h"
+#include "Utils.h"
+#include <algorithm>  // for std::max
+
+namespace AutoRTFM
+{
+
+#if (defined(__AUTORTFM) && __AUTORTFM)
+
+FAutoRTFMMetrics GAutoRTFMMetrics;
+
+void ResetAutoRTFMMetrics()
+{
+	GAutoRTFMMetrics = {};
+}
+
+FAutoRTFMMetrics GetAutoRTFMMetrics()
+{
+	return GAutoRTFMMetrics;
+}
+
+#else  // __AUTORTFM off
+
+void ResetAutoRTFMMetrics() {}
+
+FAutoRTFMMetrics GetAutoRTFMMetrics()
+{
+	return {};
+}
+
+#endif
+
+}  // namespace AutoRTFM
diff --git a/Engine/Source/Runtime/AutoRTFM/Private/MetricsPriv.h b/Engine/Source/Runtime/AutoRTFM/Private/MetricsPriv.h
new file mode 100644
--- /dev/null
+++ b/Engine/Source/Runtime/AutoRTFM/Private/MetricsPriv.h
@@ -0,0 +1,57 @@
+// Copyright Epic Games, Inc. All Rights Reserved.
+
+#pragma once
+
+#include "AutoRTFM/Metrics.h"
+
+namespace AutoRTFM
+{
+
+#if (defined(__AUTORTFM) && __AUTORTFM)
+
+/** Increments counter for number of transactions started. */
+inline void IncrementNumTransactionsStartedMetric()
+{
+	extern AutoRTFM::FAutoRTFMMetrics GAutoRTFMMetrics;
+	++GAutoRTFMMetrics.NumTransactionsStarted;
+}
+
+/** Increments counter for number of transactions committed. */
+inline void IncrementNumTransactionsCommittedMetric()
+{
+	extern AutoRTFM::FAutoRTFMMetrics GAutoRTFMMetrics;
+	++GAutoRTFMMetrics.NumTransactionsCommitted;
+}
+
+/** Increments counter for number of transactions aborted. */
+inline void IncrementNumTransactionsAbortedMetric()
+{
+	extern AutoRTFM::FAutoRTFMMetrics GAutoRTFMMetrics;
+	++GAutoRTFMMetrics.NumTransactionsAborted;
+}
+
+/** Increments counter for number of transactions aborted via user request (AbortTransaction). */
+inline void IncrementNumTransactionsAbortedByRequestMetric()
+{
+	extern AutoRTFM::FAutoRTFMMetrics GAutoRTFMMetrics;
+	++GAutoRTFMMetrics.NumTransactionsAbortedByRequest;
+}
+
+/** Increments counter for number of transactions aborted due to violating transactionality (e.g. calling a non-closed function). */
+inline void IncrementNumTransactionsAbortedByLanguageMetric()
+{
+	extern AutoRTFM::FAutoRTFMMetrics GAutoRTFMMetrics;
+	++GAutoRTFMMetrics.NumTransactionsAbortedByLanguage;
+}
+
+#else
+
+inline void IncrementNumTransactionsStartedMetric() {}
+inline void IncrementNumTransactionsCommittedMetric() {}
+inline void IncrementNumTransactionsAbortedMetric() {}
+inline void IncrementNumTransactionsAbortedByRequestMetric() {}
+inline void IncrementNumTransactionsAbortedByLanguageMetric() {}
+
+#endif
+
+}  // namespace AutoRTFM
diff --git a/Engine/Source/Runtime/AutoRTFM/Public/AutoRTFM/Metrics.h b/Engine/Source/Runtime/AutoRTFM/Public/AutoRTFM/Metrics.h
--- a/Engine/Source/Runtime/AutoRTFM/Public/AutoRTFM/Metrics.h
+++ b/Engine/Source/Runtime/AutoRTFM/Public/AutoRTFM/Metrics.h
@@ -9,18 +9,17 @@ namespace AutoRTFM
 
 struct FAutoRTFMMetrics
 {
-	int64_t NumTransactionsStarted;
-	int64_t NumTransactionsCommitted;
-
-	int64_t NumTransactionsAborted;
-	int64_t NumTransactionsAbortedByRequest;
-	int64_t NumTransactionsAbortedByLanguage;
+	uint64_t NumTransactionsStarted;
+	uint64_t NumTransactionsCommitted;
+	uint64_t NumTransactionsAborted;
+	uint64_t NumTransactionsAbortedByRequest;
+	uint64_t NumTransactionsAbortedByLanguage;
 };
 
-// reset the internal metrics tracking back to zero
+/** Reset the internal metrics tracking back to zero. */
 void ResetAutoRTFMMetrics();
 
-// get a snapshot of the current internal metrics
+/** Get a snapshot of the current internal metrics. */
 FAutoRTFMMetrics GetAutoRTFMMetrics();
 
 }
