[AutoRTFM] Remove dead code.

Functionally, this appears to be a duplicate version of the AutoRTFM::AbortTransaction API, but is never used. In fact, `autortfm_abort()` was never exposed in the AutoRTFM C API.

#Synced-CL 50048130

diff --git a/Engine/Source/Runtime/AutoRTFM/Private/API.cpp b/Engine/Source/Runtime/AutoRTFM/Private/API.cpp
--- a/Engine/Source/Runtime/AutoRTFM/Private/API.cpp
+++ b/Engine/Source/Runtime/AutoRTFM/Private/API.cpp
@@ -451,12 +451,6 @@ AUTORTFM_INTERNAL autortfm_result RTFM_autortfm_transact_then_open(void (*Uninst
 	return TransactThenOpenImpl(UninstrumentedWork, Arg);
 }
 
-AUTORTFM_INTERNAL void RTFM_autortfm_abort()
-{
-	FContext* Context = FContext::Get();
-	Context->AbortByRequestAndThrow();
-}
-
 AUTORTFM_INTERNAL void RTFM_autortfm_abort_transaction()
 {
 	FContext* const Context = FContext::Get();
@@ -814,17 +808,6 @@ extern "C" autortfm_result autortfm_transact_then_open(void (*UninstrumentedWork
 	return TransactThenOpenImpl(UninstrumentedWork, Arg);
 }
 
-extern "C" void autortfm_abort() noexcept
-{
-	if (autortfm_is_closed())
-	{
-		return RTFM_autortfm_abort();
-	}
-
-	AUTORTFM_FATAL_IF(!FContext::Get()->IsTransactional(), "The function `autortfm_abort` was called from outside a transaction");
-	FContext::Get()->AbortByRequestAndThrow();
-}
-
 extern "C" autortfm_transaction_id autortfm_current_transaction_id() noexcept
 {
 	FContext* Context = FContext::Get();
diff --git a/Engine/Source/Runtime/AutoRTFM/Private/Context.cpp b/Engine/Source/Runtime/AutoRTFM/Private/Context.cpp
--- a/Engine/Source/Runtime/AutoRTFM/Private/Context.cpp
+++ b/Engine/Source/Runtime/AutoRTFM/Private/Context.cpp
@@ -621,14 +621,6 @@ ETransactionResult FContext::Transact(void (*UninstrumentedFunction)(void*), voi
 	return Result;
 }
 
-void FContext::AbortByRequestAndThrow()
-{
-	AUTORTFM_ASSERT(Status == EContextStatus::OnTrack);
-	IncrementNumTransactionsAbortedByRequestMetric();
-	RollbackTransaction(EContextStatus::AbortedByRequest);
-	Throw();
-}
-
 void FContext::AbortByLanguageAndThrow()
 {
 	AUTORTFM_ASSERT(Status == EContextStatus::OnTrack);
diff --git a/Engine/Source/Runtime/AutoRTFM/Private/Context.h b/Engine/Source/Runtime/AutoRTFM/Private/Context.h
--- a/Engine/Source/Runtime/AutoRTFM/Private/Context.h
+++ b/Engine/Source/Runtime/AutoRTFM/Private/Context.h
@@ -38,8 +38,6 @@ public:
 
 	EContextStatus CallClosedNest(void (*ClosedFunction)(void* Arg), void* Arg);
 
-	void AbortByRequestAndThrow();
-
 	// Open API - no throw
 	void StartTransaction();
 
