[AutoRTFM] Fix metrics to report correct values.

Previously, there were no tests for our metrics, so we were double-counting some events and missing others. Also, we did not maintain a count of top-level transactions, but this is relevant information as our histogram is meant to report the size of top-level transactions.

#Synced-CL 50048130

diff --git a/Engine/Source/Programs/AutoRTFMTests/Private/MetricsTests.cpp b/Engine/Source/Programs/AutoRTFMTests/Private/MetricsTests.cpp
new file mode 100644
--- /dev/null
+++ b/Engine/Source/Programs/AutoRTFMTests/Private/MetricsTests.cpp
@@ -0,0 +1,83 @@
+// Copyright Epic Games, Inc. All Rights Reserved.
+
+#include "AutoRTFM.h"
+#include "AutoRTFMTestUtils.h"
+#include "AutoRTFM/Metrics.h"
+#include "AutoRTFM/Testing.h"
+#include "Catch2Includes.h"
+
+TEST_CASE("Metrics.TransactionCounts")
+{
+	AUTORTFM_SCOPED_DISABLE_RETRY();
+
+	AutoRTFM::ResetAutoRTFMMetrics();
+
+	SECTION("Top-level")
+	{
+		AutoRTFM::Transact([] {});
+		AutoRTFM::Transact([] { AutoRTFM::AbortTransaction(); });
+
+		AutoRTFM::FAutoRTFMMetrics Metrics = AutoRTFM::GetAutoRTFMMetrics();
+		REQUIRE(Metrics.NumTopLevelTransactionsStarted == 2);
+		REQUIRE(Metrics.NumTransactionsStarted == 2);
+		REQUIRE(Metrics.NumTransactionsCommitted == 1);
+		REQUIRE(Metrics.NumTransactionsAborted == 1);
+		REQUIRE(Metrics.NumTransactionsAbortedByRequest == 1);
+		REQUIRE(Metrics.NumTransactionsAbortedByLanguage == 0);
+	}
+
+	SECTION("Nesting")
+	{
+		AutoRTFM::Transact([]
+		{
+			AutoRTFM::Transact([] {});
+			AutoRTFM::Transact([] {});
+			AutoRTFM::Transact([]
+			{
+				AutoRTFM::Transact([] {});
+			});
+		});
+
+		AutoRTFM::FAutoRTFMMetrics Metrics = AutoRTFM::GetAutoRTFMMetrics();
+		REQUIRE(Metrics.NumTopLevelTransactionsStarted == 1);
+		REQUIRE(Metrics.NumTransactionsStarted == 5);
+		REQUIRE(Metrics.NumTransactionsCommitted == 5);
+		REQUIRE(Metrics.NumTransactionsAborted == 0);
+		REQUIRE(Metrics.NumTransactionsAbortedByRequest == 0);
+		REQUIRE(Metrics.NumTransactionsAbortedByLanguage == 0);
+	}
+
+	SECTION("Scoped rollback")
+	{
+		AutoRTFM::TransactThenOpen([]
+		{
+			AutoRTFM::ForTheRuntime::StartTransaction();
+			AutoRTFM::ForTheRuntime::RollbackTransaction();
+			AutoRTFM::ForTheRuntime::ClearTransactionStatus();
+		});
+		AutoRTFM::FAutoRTFMMetrics Metrics = AutoRTFM::GetAutoRTFMMetrics();
+		REQUIRE(Metrics.NumTopLevelTransactionsStarted == 1);
+		REQUIRE(Metrics.NumTransactionsStarted == 2);
+		REQUIRE(Metrics.NumTransactionsCommitted == 1);
+		REQUIRE(Metrics.NumTransactionsAborted == 1);
+		REQUIRE(Metrics.NumTransactionsAbortedByRequest == 0);
+		REQUIRE(Metrics.NumTransactionsAbortedByLanguage == 0);
+	}
+
+	SECTION("Scoped commit")
+	{
+		AutoRTFM::TransactThenOpen([]
+		{
+			AutoRTFM::ForTheRuntime::StartTransaction();
+			AutoRTFM::ForTheRuntime::CommitTransaction();
+			AutoRTFM::ForTheRuntime::ClearTransactionStatus();
+		});
+		AutoRTFM::FAutoRTFMMetrics Metrics = AutoRTFM::GetAutoRTFMMetrics();
+		REQUIRE(Metrics.NumTopLevelTransactionsStarted == 1);
+		REQUIRE(Metrics.NumTransactionsStarted == 2);
+		REQUIRE(Metrics.NumTransactionsCommitted == 2);
+		REQUIRE(Metrics.NumTransactionsAborted == 0);
+		REQUIRE(Metrics.NumTransactionsAbortedByRequest == 0);
+		REQUIRE(Metrics.NumTransactionsAbortedByLanguage == 0);
+	}
+}
diff --git a/Engine/Source/Runtime/AutoRTFM/Private/Context.cpp b/Engine/Source/Runtime/AutoRTFM/Private/Context.cpp
--- a/Engine/Source/Runtime/AutoRTFM/Private/Context.cpp
+++ b/Engine/Source/Runtime/AutoRTFM/Private/Context.cpp
@@ -93,16 +93,6 @@ bool FContext::IsRetrying() const
 	return GetStatus() == EContextStatus::AbortedByCascadingRetry;
 }
 
-void FContext::MaterializeDeferredTransactions()
-{
-	uint64_t NumToAllocate = GetNumDeferredTransactions();
-	NumDeferredTransactions = 0;
-	for (uint64_t I = 0; I < NumToAllocate; ++I)
-	{
-		StartNonDeferredTransaction();
-	}
-}
-
 void FContext::StartTransaction()
 {
 	NumDeferredTransactions += 1;
@@ -112,21 +102,23 @@ void FContext::StartTransaction()
 	IncrementNumTransactionsStartedMetric();
 }
 
-void FContext::StartNonDeferredTransaction()
+void FContext::MaterializeDeferredTransactions()
 {
-	AUTORTFM_ASSERT(GetNumDeferredTransactions() == 0);
-	AUTORTFM_ENSURE_MSG(
-		CurrentTransaction, "FContext::StartNonDeferredTransaction() can only be called within a scoped transaction");
-
-	PushTransaction(
-		/* Closed */ false,
-		/* bIsScoped */ false,
-		/* StackRange */ CurrentTransaction->GetStackRange());
-
-	// This form of transaction is always ultimately within a scoped Transact
-	AUTORTFM_ASSERT(Status == EContextStatus::OnTrack);
+	if (uint64_t NumToAllocate = GetNumDeferredTransactions(); NumToAllocate > 0)
+	{
+		// This form of transaction is always ultimately within a scoped Transact
+		AUTORTFM_ASSERT(Status == EContextStatus::OnTrack);
+		AUTORTFM_ENSURE_MSG(
+			CurrentTransaction, "FContext::MaterializeDeferredTransactions() can only be called within a scoped transaction");
 
-	IncrementNumTransactionsStartedMetric();
+		NumDeferredTransactions = 0;
+		for (uint64_t I = 0; I < NumToAllocate; ++I)
+		{
+			PushTransaction(/* Closed */ false,
+				/* bIsScoped */ false,
+				/* StackRange */ CurrentTransaction->GetStackRange());
+		}
+	}
 }
 
 ETransactionResult FContext::CommitTransaction()
@@ -139,6 +131,7 @@ ETransactionResult FContext::CommitTransaction()
 	{
 		// The optimization worked! We didn't need to allocate an FTransaction for this.
 		NumDeferredTransactions -= 1;
+		IncrementNumTransactionsCommittedMetric();
 	}
 	else
 	{
@@ -171,8 +164,6 @@ ETransactionResult FContext::CommitTransaction()
 		PopTransaction();
 	}
 
-	IncrementNumTransactionsCommittedMetric();
-
 	return Result;
 }
 
@@ -207,6 +198,8 @@ void FContext::RollbackTransaction(EContextStatus NewStatus)
 
 void FContext::AbortTransaction(EContextStatus NewStatus)
 {
+	AUTORTFM_ASSERT(NewStatus != EContextStatus::AbortedByLanguage);  // use AbortByLanguageAndThrow instead
+	IncrementNumTransactionsAbortedByRequestMetric();
 	RollbackTransaction(NewStatus);
 	Throw();
 }
@@ -385,8 +378,13 @@ ETransactionResult FContext::Transact(void (*UninstrumentedFunction)(void*), voi
 
 	ETransactionResult Result = ETransactionResult::Committed;  // Initialize to something to make the compiler happy.
 
+	IncrementNumTransactionsStartedMetric();
+
 	if (!CurrentTransaction)
 	{
+		// This is a top-level transaction.
+		IncrementNumTopLevelTransactionsStartedMetric();
+
 		// If exceptions are enabled, then ensure that the transaction is automatically committed if
 		// an exception is thrown inside the transaction and the handler is outside the transaction.
 		struct FAutoCommitter final
diff --git a/Engine/Source/Runtime/AutoRTFM/Private/Context.h b/Engine/Source/Runtime/AutoRTFM/Private/Context.h
--- a/Engine/Source/Runtime/AutoRTFM/Private/Context.h
+++ b/Engine/Source/Runtime/AutoRTFM/Private/Context.h
@@ -177,7 +177,6 @@ private:
 	~FContext();
 
 	void MaterializeDeferredTransactions();
-	void StartNonDeferredTransaction();
 
 	void PushCallNest(FCallNest* NewCallNest);
 	void PopCallNest();
diff --git a/Engine/Source/Runtime/AutoRTFM/Private/ContextInlines.h b/Engine/Source/Runtime/AutoRTFM/Private/ContextInlines.h
--- a/Engine/Source/Runtime/AutoRTFM/Private/ContextInlines.h
+++ b/Engine/Source/Runtime/AutoRTFM/Private/ContextInlines.h
@@ -6,6 +6,7 @@
 
 #include "AutoRTFM/Defines.h"
 #include "Context.h"
+#include "MetricsPriv.h"
 #include "TransactionInlines.h"
 
 namespace AutoRTFM
@@ -64,14 +65,14 @@ UE_AUTORTFM_FORCEINLINE bool FContext::AttemptToCommitTransaction(FTransaction*
 
 	Status = EContextStatus::Committing;
 
-	const bool bResult = Transaction->AttemptToCommit();
-
-	if (bResult)
+	if (Transaction->AttemptToCommit())
 	{
+		IncrementNumTransactionsCommittedMetric();
 		Status = EContextStatus::OnTrack;
+		return true;
 	}
 
-	return bResult;
+	return false;
 }
 
 }  // namespace AutoRTFM
diff --git a/Engine/Source/Runtime/AutoRTFM/Private/MetricsPriv.h b/Engine/Source/Runtime/AutoRTFM/Private/MetricsPriv.h
--- a/Engine/Source/Runtime/AutoRTFM/Private/MetricsPriv.h
+++ b/Engine/Source/Runtime/AutoRTFM/Private/MetricsPriv.h
@@ -7,7 +7,12 @@
 namespace AutoRTFM
 {
 
-#if (defined(__AUTORTFM) && __AUTORTFM)
+/** Increments counter for number of top-level transactions started. */
+inline void IncrementNumTopLevelTransactionsStartedMetric()
+{
+	extern AutoRTFM::FAutoRTFMMetrics GAutoRTFMMetrics;
+	++GAutoRTFMMetrics.NumTopLevelTransactionsStarted;
+}
 
 /** Increments counter for number of transactions started. */
 inline void IncrementNumTransactionsStartedMetric()
@@ -44,14 +49,4 @@ inline void IncrementNumTransactionsAbortedByLanguageMetric()
 	++GAutoRTFMMetrics.NumTransactionsAbortedByLanguage;
 }
 
-#else
-
-inline void IncrementNumTransactionsStartedMetric() {}
-inline void IncrementNumTransactionsCommittedMetric() {}
-inline void IncrementNumTransactionsAbortedMetric() {}
-inline void IncrementNumTransactionsAbortedByRequestMetric() {}
-inline void IncrementNumTransactionsAbortedByLanguageMetric() {}
-
-#endif
-
 }  // namespace AutoRTFM
diff --git a/Engine/Source/Runtime/AutoRTFM/Public/AutoRTFM/Metrics.h b/Engine/Source/Runtime/AutoRTFM/Public/AutoRTFM/Metrics.h
--- a/Engine/Source/Runtime/AutoRTFM/Public/AutoRTFM/Metrics.h
+++ b/Engine/Source/Runtime/AutoRTFM/Public/AutoRTFM/Metrics.h
@@ -9,6 +9,7 @@ namespace AutoRTFM
 
 struct FAutoRTFMMetrics
 {
+	uint64_t NumTopLevelTransactionsStarted;
 	uint64_t NumTransactionsStarted;
 	uint64_t NumTransactionsCommitted;
 	uint64_t NumTransactionsAborted;
