[AutoRTFM] Change AutoRTFM::IsCommittingOrAborting() to match the old behavior

This used to return true for static initializers, further compounding the need to rename this function.

#Synced-CL 50210011

diff --git a/Engine/Source/Runtime/AutoRTFM/Public/AutoRTFM.h b/Engine/Source/Runtime/AutoRTFM/Public/AutoRTFM.h
--- a/Engine/Source/Runtime/AutoRTFM/Public/AutoRTFM.h
+++ b/Engine/Source/Runtime/AutoRTFM/Public/AutoRTFM.h
@@ -404,8 +404,8 @@ UE_AUTORTFM_CRITICAL_INLINE bool IsRetrying()
 	return ContextStatus() == EContextStatus::Retrying;
 }
 
-// Tells us if we are currently committing or aborting a transaction.
-// This will return true when inside an on-abort, on-commit, on-complete or on-retry handler.
+// Returns true when inside an on-abort, on-commit, on-complete, on-retry handler
+// or inside a static initializer.
 // TODO(SOL-8967): This function name is misleading. Rename to something sensible.
 UE_AUTORTFM_CRITICAL_INLINE bool IsCommittingOrAborting()
 {
@@ -415,11 +415,11 @@ UE_AUTORTFM_CRITICAL_INLINE bool IsCommittingOrAborting()
 		case AutoRTFM::EContextStatus::Aborting:
 		case AutoRTFM::EContextStatus::Completing:
 		case AutoRTFM::EContextStatus::Retrying:
+		case AutoRTFM::EContextStatus::InStaticLocalInitializer:
 			return true;
 		default:
 		case AutoRTFM::EContextStatus::Idle:
 		case AutoRTFM::EContextStatus::OnTrack:
-		case AutoRTFM::EContextStatus::InStaticLocalInitializer:
 			return false;
 	}
 }
