[AutoRTFM] Minor code clean up.

* Rename FContext::CurrentTransaction to CurrentMaterializedTransaction. Acts as a reminder to the reader that this may not be the current "deferred" context.
* Added a bunch of documentation on the FContext public methods, especially regarding deferred / materialized contexts.
* Killed FTransaction::IsNested() for just using Parent. Done for consistency.
* Killed FTransaction::IsNestedWithin(). Was not used.
* Replaced the templated AutoRTFM::HasAborted() with two distinct overloads.

#Synced-CL 50245569

diff --git a/Engine/Source/Runtime/AutoRTFM/Private/Context.cpp b/Engine/Source/Runtime/AutoRTFM/Private/Context.cpp
--- a/Engine/Source/Runtime/AutoRTFM/Private/Context.cpp
+++ b/Engine/Source/Runtime/AutoRTFM/Private/Context.cpp
@@ -63,7 +63,7 @@ FContext::~FContext()
 void FContext::StartTransaction()
 {
 	AUTORTFM_ASSERT(Status == EContextStatus::OnTrack);
-	AUTORTFM_ENSURE_MSG(CurrentTransaction, "FContext::StartTransaction() cannot be used for the outermost transaction");
+	AUTORTFM_ENSURE_MSG(CurrentMaterializedTransaction, "FContext::StartTransaction() cannot be used for the outermost transaction");
 
 	NumDeferredTransactions += 1;
 
@@ -76,12 +76,12 @@ void FContext::MaterializeDeferredTransactions()
 	{
 		// This form of transaction is always ultimately within a scoped Transact
 		AUTORTFM_ASSERT(Status == EContextStatus::OnTrack);
-		AUTORTFM_ENSURE_MSG(
-			CurrentTransaction, "FContext::MaterializeDeferredTransactions() can only be called within a scoped transaction");
+		AUTORTFM_ENSURE_MSG(CurrentMaterializedTransaction,
+			"FContext::MaterializeDeferredTransactions() can only be called within a scoped transaction");
 		NumDeferredTransactions = 0;
 		for (uint64_t I = 0; I < NumToAllocate; ++I)
 		{
-			PushTransaction(CurrentTransaction->GetStackRange());
+			PushTransaction(CurrentMaterializedTransaction->GetStackRange());
 		}
 	}
 }
@@ -90,18 +90,20 @@ void FContext::CommitTransaction()
 {
 	AUTORTFM_ASSERT(Status == EContextStatus::OnTrack);
 
+	IncrementNumTransactionsCommittedMetric();
+
 	if (NumDeferredTransactions)
 	{
 		NumDeferredTransactions--;
-		IncrementNumTransactionsCommittedMetric();
 		return;
 	}
 
+	AUTORTFM_ASSERT(CurrentMaterializedTransaction);
+
 	Status = EContextStatus::Committing;
 	{
-		IncrementNumTransactionsCommittedMetric();
 		SetCallNestTransactionNonCascadingStatus(ETransactionStatus::Committed);
-		CurrentTransaction->Commit();
+		CurrentMaterializedTransaction->Commit();
 		PopTransaction();
 	}
 	Status = EContextStatus::OnTrack;
@@ -120,7 +122,7 @@ void FContext::AbortTransaction(ETransactionStatus Reason)
 		return;
 	}
 
-	AUTORTFM_ASSERT(CurrentTransaction);
+	AUTORTFM_ASSERT(CurrentMaterializedTransaction);
 
 	switch (Reason)
 	{
@@ -142,9 +144,9 @@ void FContext::AbortTransaction(ETransactionStatus Reason)
 		SetCallNestTransactionCascadingStatus(Reason);
 
 		// Immediately abort and pop the entire transaction stack.
-		while (CurrentTransaction)
+		while (CurrentMaterializedTransaction)
 		{
-			CurrentTransaction->Abort();
+			CurrentMaterializedTransaction->Abort();
 			PopTransaction();
 		}
 	}
@@ -155,7 +157,7 @@ void FContext::AbortTransaction(ETransactionStatus Reason)
 		SetCallNestTransactionNonCascadingStatus(Reason);
 
 		// Abort and pop the top most transaction.
-		CurrentTransaction->Abort();
+		CurrentMaterializedTransaction->Abort();
 		PopTransaction();
 	}
 
@@ -174,9 +176,11 @@ void FContext::SetCallNestTransactionCascadingStatus(ETransactionStatus Transact
 void FContext::SetCallNestTransactionNonCascadingStatus(ETransactionStatus TransactionStatus)
 {
 	AUTORTFM_ASSERT_DEBUG(!IsCascading(TransactionStatus));
+	AUTORTFM_ASSERT(CurrentMaterializedTransaction);
+
 	for (FCallNest* Nest = CurrentNest; Nest; Nest = Nest->Parent)
 	{
-		if (Nest->Transaction == CurrentTransaction->Identifier())
+		if (Nest->Transaction == CurrentMaterializedTransaction->Identifier())
 		{
 			Nest->Status = TransactionStatus;
 		}
@@ -205,7 +209,8 @@ ETransactionStatus FContext::CallClosedNest(void (*ClosedFunction)(void* Arg), v
 
 void FContext::PushCallNest()
 {
-	CurrentNest = CallNestPool.Take(CurrentNest, CurrentTransaction->Identifier());
+	MaterializeDeferredTransactions();
+	CurrentNest = CallNestPool.Take(CurrentNest, CurrentMaterializedTransaction->Identifier());
 }
 
 void FContext::PopCallNest()
@@ -223,10 +228,10 @@ FTransaction* FContext::PushTransaction(FStackRange StackRange)
 
 	FTransaction* NewTransaction = TransactionPool.Take(this);
 	NewTransaction->Initialize(
-		/* Parent */ CurrentTransaction,
+		/* Parent */ CurrentMaterializedTransaction,
 		/* StackRange */ StackRange);
 
-	CurrentTransaction = NewTransaction;
+	CurrentMaterializedTransaction = NewTransaction;
 
 	// Collect stats that we've got a new transaction.
 	Stats.Collect<EStatsKind::Transaction>();
@@ -237,9 +242,9 @@ FTransaction* FContext::PushTransaction(FStackRange StackRange)
 void FContext::PopTransaction()
 {
 	AUTORTFM_ASSERT(NumDeferredTransactions == 0);
-	AUTORTFM_ASSERT(CurrentTransaction != nullptr);
-	FTransaction* OldTransaction = CurrentTransaction;
-	CurrentTransaction = CurrentTransaction->GetParent();
+	AUTORTFM_ASSERT(CurrentMaterializedTransaction != nullptr);
+	FTransaction* OldTransaction = CurrentMaterializedTransaction;
+	CurrentMaterializedTransaction = CurrentMaterializedTransaction->GetParent();
 	TransactionPool.Return(OldTransaction);
 }
 
@@ -319,21 +324,21 @@ private:
 
 ETransactionResult FContext::Transact(void (*UninstrumentedFunction)(void*), void (*InstrumentedFunction)(void*), void* Arg)
 {
-	if (AUTORTFM_UNLIKELY(EContextStatus::Committing == Status))
-	{
-		return ETransactionResult::AbortedByTransactDuringCommit;
-	}
-	if (AUTORTFM_UNLIKELY(EContextStatus::Aborting == Status))
+	switch (Status)
 	{
-		return ETransactionResult::AbortedByTransactDuringAbort;
-	}
-	if (AUTORTFM_UNLIKELY(EContextStatus::Retrying == Status))
-	{
-		return ETransactionResult::AbortedByTransactDuringRetry;
-	}
-	if (AUTORTFM_UNLIKELY(EContextStatus::Completing == Status))
-	{
-		return ETransactionResult::AbortedByTransactDuringCompletion;
+		[[unlikely]] case EContextStatus::Committing:
+			return ETransactionResult::AbortedByTransactDuringCommit;
+		[[unlikely]] case EContextStatus::Aborting:
+			return ETransactionResult::AbortedByTransactDuringAbort;
+		[[unlikely]] case EContextStatus::Retrying:
+			return ETransactionResult::AbortedByTransactDuringRetry;
+		[[unlikely]] case EContextStatus::Completing:
+			return ETransactionResult::AbortedByTransactDuringCompletion;
+
+		case EContextStatus::Idle:
+		case EContextStatus::OnTrack:
+		case EContextStatus::InStaticLocalInitializer:
+			break;
 	}
 
 	if (!InstrumentedFunction)
@@ -350,7 +355,7 @@ ETransactionResult FContext::Transact(void (*UninstrumentedFunction)(void*), voi
 
 	IncrementNumTransactionsStartedMetric();
 
-	if (CurrentTransaction)
+	if (CurrentMaterializedTransaction)
 	{
 		// Nested transaction
 		AUTORTFM_ASSERT(Status == EContextStatus::OnTrack);
@@ -371,18 +376,19 @@ ETransactionResult FContext::Transact(void (*UninstrumentedFunction)(void*), voi
 	AUTORTFM_ASSERT(Stack.Contains(&TransactStackStart));
 	FStackRange const StackRange{Stack.Low, &TransactStackStart};
 
-	FTransaction* const ParentTransaction = CurrentTransaction;
+	FTransaction* const ParentTransaction = CurrentMaterializedTransaction;
 
 	ETransactionResult Result{};
 
-	bool const bShouldRetryTransaction = CurrentTransaction ? AutoRTFM::ForTheRuntime::ShouldRetryNestedTransactionsToo()
-															: AutoRTFM::ForTheRuntime::ShouldRetryNonNestedTransactions();
+	bool const bShouldRetryTransaction = CurrentMaterializedTransaction
+										   ? AutoRTFM::ForTheRuntime::ShouldRetryNestedTransactionsToo()
+										   : AutoRTFM::ForTheRuntime::ShouldRetryNonNestedTransactions();
 	bool bTriedToRunOnce = false;
 
 	for (;;)
 	{
 		Status = EContextStatus::OnTrack;
-		AUTORTFM_ASSERT(CurrentTransaction == ParentTransaction);
+		AUTORTFM_ASSERT(CurrentMaterializedTransaction == ParentTransaction);
 
 		FTransaction* const Transaction = PushTransaction(StackRange);
 		PushCallNest();
@@ -396,7 +402,7 @@ ETransactionResult FContext::Transact(void (*UninstrumentedFunction)(void*), voi
 
 		if (AUTORTFM_UNLIKELY(bShouldRetryTransaction) && !bTriedToRunOnce && TransactionStatus == ETransactionStatus::Executing)
 		{
-			AUTORTFM_ASSERT(CurrentTransaction == Transaction);
+			AUTORTFM_ASSERT(CurrentMaterializedTransaction == Transaction);
 			bTriedToRunOnce = true;
 			AbortTransaction(ETransactionStatus::AbortedByRequest);
 			if (!ParentTransaction)
@@ -409,7 +415,7 @@ ETransactionResult FContext::Transact(void (*UninstrumentedFunction)(void*), voi
 		switch (TransactionStatus)
 		{
 			case ETransactionStatus::Executing:
-				AUTORTFM_ASSERT(CurrentTransaction == Transaction);
+				AUTORTFM_ASSERT(CurrentMaterializedTransaction == Transaction);
 				CommitTransaction();
 				Result = ETransactionResult::Committed;
 				break;
@@ -427,7 +433,7 @@ ETransactionResult FContext::Transact(void (*UninstrumentedFunction)(void*), voi
 				break;
 
 			case ETransactionStatus::AbortedByCascadingAbort:
-				AUTORTFM_ASSERT(CurrentTransaction == nullptr);
+				AUTORTFM_ASSERT(CurrentMaterializedTransaction == nullptr);
 				if (CurrentNest)
 				{
 					Throw();
@@ -436,7 +442,7 @@ ETransactionResult FContext::Transact(void (*UninstrumentedFunction)(void*), voi
 				break;
 
 			case ETransactionStatus::AbortedByCascadingRetry:
-				AUTORTFM_ASSERT(CurrentTransaction == nullptr);
+				AUTORTFM_ASSERT(CurrentMaterializedTransaction == nullptr);
 				if (CurrentNest)
 				{
 					Throw();
@@ -473,7 +479,7 @@ ETransactionResult FContext::Transact(void (*UninstrumentedFunction)(void*), voi
 void FContext::EndScopedTransaction()
 {
 	AUTORTFM_ASSERT(Status == EContextStatus::OnTrack);
-	if (CurrentTransaction == nullptr)
+	if (CurrentMaterializedTransaction == nullptr)
 	{
 		// No more transactions in flight.
 		RetryTasks.Reset();
@@ -485,7 +491,7 @@ void FContext::EndScopedTransaction()
 void FContext::CallCompletionTasks()
 {
 	AUTORTFM_ASSERT(Status == EContextStatus::OnTrack);
-	AUTORTFM_ASSERT(CurrentTransaction == nullptr);
+	AUTORTFM_ASSERT(CurrentMaterializedTransaction == nullptr);
 	Status = EContextStatus::Completing;
 	CompletionTasks.RemoveEachForward([](TTask<void()>& Task) { Task(); });
 	Status = EContextStatus::OnTrack;
@@ -497,7 +503,7 @@ void FContext::Reset()
 
 	CurrentThreadId = FThreadID::Invalid;
 	Stack = {};
-	CurrentTransaction = nullptr;
+	CurrentMaterializedTransaction = nullptr;
 	CurrentNest = nullptr;
 	Status = EContextStatus::Idle;
 	StackLocalInitializerDepth = 0;
diff --git a/Engine/Source/Runtime/AutoRTFM/Private/Context.h b/Engine/Source/Runtime/AutoRTFM/Private/Context.h
--- a/Engine/Source/Runtime/AutoRTFM/Private/Context.h
+++ b/Engine/Source/Runtime/AutoRTFM/Private/Context.h
@@ -38,12 +38,35 @@ public:
 
 	ETransactionStatus CallClosedNest(void (*ClosedFunction)(void* Arg), void* Arg);
 
+	// Begins an unscoped transaction.
+	// It is common for a transaction to be started and then immediately aborted
+	// without actually being used. As a performance optimization,
+	// StartTransaction() increments NumDeferredTransactions, which represents
+	// the number of "deferred" transactions at the top of the transaction
+	// stack. Deferred transactions will be lazily backed by an FTransaction on
+	// the first call to MaterializeDeferredTransactions() or
+	// GetCurrentTransaction().
 	void StartTransaction();
+
+	// Ends a scoped or unscoped transaction, undoing all writes recorded by the
+	// transaction and calling all registered on-abort handlers.
+	// If the current transaction is materialized, then the FTransaction is
+	// popped from the transaction stack, otherwise NumDeferredTransactions is
+	// decremented.
 	void AbortTransaction(ETransactionStatus Reason);
+
+	// Convenience function for calling AbortTransaction() followed by Throw().
 	void AbortTransactionAndThrow(ETransactionStatus Reason);
 
+	// Ends a scoped or unscoped transaction, calling all registered on-commit
+	// handlers.
+	// If the current transaction is materialized, then the FTransaction is
+	// popped from the transaction stack, otherwise NumDeferredTransactions is
+	// decremented.
 	void CommitTransaction();
 
+	// Throws if the current (most-nested) FCallNest has a status that is not
+	// ETransactionStatus::Executing.
 	void MaybeThrow();
 
 	// Record that a write is about to occur at the given LogicalAddress of Size bytes.
@@ -59,35 +82,54 @@ public:
 	inline void DeferUntilComplete(TTask<void()>&& Callback);
 	inline void DeferUntilRetry(TTask<void()>&& Callback);
 
-	// The rest of this is internalish.
+	// The number of deferred transactions that have yet to be materialized into
+	// a FTransaction.
 	inline uint64_t GetNumDeferredTransactions() const
 	{
 		return NumDeferredTransactions;
 	}
+
+	// Returns true if the current (most-nested) transaction is deferred.
+	// For more information about deferred transactions, see StartTransaction()
+	// and AbortTransaction().
 	inline bool MustMaterializeDeferredTransactions() const
 	{
 		return (AUTORTFM_UNLIKELY(0 < GetNumDeferredTransactions()));
 	}
+
+	// Returns the current (most-nested) transaction, materializing the
+	// transaction stack if one or more transactions are deferred.
+	// For more information about deferred transactions, see StartTransaction()
+	// and AbortTransaction().
 	inline FTransaction* GetCurrentTransaction()
 	{
 		if (MustMaterializeDeferredTransactions())
 		{
 			MaterializeDeferredTransactions();
 		}
-		return CurrentTransaction;
+		return CurrentMaterializedTransaction;
 	}
+
+	// Returns the most-nested materialized transaction.
+	// For more information about deferred transactions, see StartTransaction()
+	// and AbortTransaction().
 	inline FTransaction* GetMaterializedTransaction()
 	{
-		return CurrentTransaction;
+		return CurrentMaterializedTransaction;
 	}
+
+	// Returns true if the caller's thread is the same as the thread associated
+	// with the transaction stack.
 	inline bool IsContextThread() const
 	{
 		return CurrentThreadId == FThreadID::GetCurrent();
 	}
+
 	inline EContextStatus GetStatus() const
 	{
 		return IsContextThread() ? Status : EContextStatus::Idle;
 	}
+
 	inline bool IsTransactional() const
 	{
 		return GetStatus() == EContextStatus::OnTrack;
@@ -136,6 +178,7 @@ public:
 	{
 		return TaskPool;
 	}
+
 	inline FTransactionPool& GetTransactionPool()
 	{
 		return TransactionPool;
@@ -190,7 +233,7 @@ private:
 	// We defer allocating FTransactions at the top of the transaction stack.
 	// This allows us to make starting a transaction in the open a load, some math, and a store.
 	uint64_t NumDeferredTransactions{0};
-	FTransaction* CurrentTransaction{nullptr};
+	FTransaction* CurrentMaterializedTransaction{nullptr};
 	FCallNest* CurrentNest{nullptr};
 
 	FStackRange Stack;
diff --git a/Engine/Source/Runtime/AutoRTFM/Private/Transaction.cpp b/Engine/Source/Runtime/AutoRTFM/Private/Transaction.cpp
--- a/Engine/Source/Runtime/AutoRTFM/Private/Transaction.cpp
+++ b/Engine/Source/Runtime/AutoRTFM/Private/Transaction.cpp
@@ -111,7 +111,7 @@ void FTransaction::Abort()
 	CollectStats();
 
 	UpdatePeakMemoryUsageMetrics();
-	if (!IsNested())
+	if (!Parent)
 	{
 		UpdateMemoryUsageHistogram();
 	}
@@ -140,7 +140,7 @@ void FTransaction::Commit()
 	Stats.Collect<EStatsKind::Commit>();
 	CollectStats();
 
-	if (IsNested())
+	if (Parent)
 	{
 		CommitNested();
 	}
diff --git a/Engine/Source/Runtime/AutoRTFM/Private/Transaction.h b/Engine/Source/Runtime/AutoRTFM/Private/Transaction.h
--- a/Engine/Source/Runtime/AutoRTFM/Private/Transaction.h
+++ b/Engine/Source/Runtime/AutoRTFM/Private/Transaction.h
@@ -39,31 +39,11 @@ class AUTORTFM_DISABLE_NO_SANITIZE FTransaction final
 public:
 	void Initialize(FTransaction* InParent, FStackRange InStackRange);
 
-	bool IsNested() const
-	{
-		return !!Parent;
-	}
 	FTransaction* GetParent() const
 	{
 		return Parent;
 	}
 
-	// This should just use type displays or ranges. Maybe ranges could even work out great.
-	bool IsNestedWithin(const FTransaction* Other) const
-	{
-		for (const FTransaction* Current = this;; Current = Current->Parent)
-		{
-			if (!Current)
-			{
-				return false;
-			}
-			else if (Current == Other)
-			{
-				return true;
-			}
-		}
-	}
-
 	void DeferUntilCommit(TTask<void()>&&);
 	void DeferUntilPreAbort(TTask<void()>&&);
 	void DeferUntilAbort(TTask<void()>&&);
diff --git a/Engine/Source/Runtime/AutoRTFM/Public/AutoRTFM.h b/Engine/Source/Runtime/AutoRTFM/Public/AutoRTFM.h
--- a/Engine/Source/Runtime/AutoRTFM/Public/AutoRTFM.h
+++ b/Engine/Source/Runtime/AutoRTFM/Public/AutoRTFM.h
@@ -104,6 +104,28 @@ enum class ETransactionResult
 	AbortedByCascadingAbort = autortfm_transaction_aborted_by_cascading_abort,
 };
 
+// Returns true if the given ETransactionResult represents a transaction that
+// has aborted.
+inline bool HasAborted(ETransactionResult State)
+{
+	// Note: Intentionally verbose so that the compiler will error if new enumerators are added.
+	switch (State)
+	{
+		case ETransactionResult::Committed:
+			return false;
+
+		case ETransactionResult::AbortedByRequest:
+		case ETransactionResult::AbortedByLanguage:
+		case ETransactionResult::AbortedByTransactDuringCommit:
+		case ETransactionResult::AbortedByTransactDuringAbort:
+		case ETransactionResult::AbortedByTransactDuringRetry:
+		case ETransactionResult::AbortedByTransactDuringCompletion:
+		case ETransactionResult::AbortedByCascadingAbort:
+			break;
+	}
+	return true;
+}
+
 // The transaction status provides information on how a transaction is currently executing or completed.
 // ETransactionStatus is a superset of ETransactionResult.
 enum class ETransactionStatus
@@ -144,37 +166,39 @@ enum class ETransactionStatus
 	AbortedByCascadingRetry = autortfm_transaction_aborted_by_cascading_retry,
 };
 
-// Returns true if the given ETransactionStatus or ETransactionResult represents a transaction
-// that has aborted.
-template <typename T>
-requires(std::is_same_v<T, ETransactionStatus> || std::is_same_v<T, ETransactionResult>)
-inline bool HasAborted(T State)
+// Returns true if the given ETransactionStatus represents a transaction that has aborted.
+inline bool HasAborted(ETransactionStatus State)
 {
-	if constexpr (std::is_same_v<T, ETransactionStatus>)
-	{
-		if (State == T::AbortedByCascadingRetry)
-		{
-			return true;
-		}
-	}
+	// Note: Intentionally verbose so that the compiler will error if new enumerators are added.
 	switch (State)
 	{
-		case T::AbortedByRequest:
-		case T::AbortedByLanguage:
-		case T::AbortedByTransactDuringCommit:
-		case T::AbortedByTransactDuringAbort:
-		case T::AbortedByCascadingAbort:
-			return true;
-		default:
+		case ETransactionStatus::Executing:
+		case ETransactionStatus::Committed:
 			return false;
+
+		case ETransactionStatus::AbortedByRequest:
+		case ETransactionStatus::AbortedByLanguage:
+		case ETransactionStatus::AbortedByTransactDuringCommit:
+		case ETransactionStatus::AbortedByTransactDuringAbort:
+		case ETransactionStatus::AbortedByTransactDuringRetry:
+		case ETransactionStatus::AbortedByTransactDuringCompletion:
+		case ETransactionStatus::AbortedByCascadingAbort:
+		case ETransactionStatus::AbortedByCascadingRetry:
+			break;
 	}
+	return true;
 }
 
 // Returns true if Status is AbortedByCascadingAbort or AbortedByCascadingRetry.
 inline bool IsCascading(ETransactionStatus Status)
 {
+	// Note: Intentionally verbose so that the compiler will error if new enumerators are added.
 	switch (Status)
 	{
+		case ETransactionStatus::AbortedByCascadingAbort:
+		case ETransactionStatus::AbortedByCascadingRetry:
+			return true;
+
 		case ETransactionStatus::Executing:
 		case ETransactionStatus::Committed:
 		case ETransactionStatus::AbortedByRequest:
@@ -183,11 +207,9 @@ inline bool IsCascading(ETransactionStatus Status)
 		case ETransactionStatus::AbortedByTransactDuringAbort:
 		case ETransactionStatus::AbortedByTransactDuringRetry:
 		case ETransactionStatus::AbortedByTransactDuringCompletion:
-			return false;
-		case ETransactionStatus::AbortedByCascadingAbort:
-		case ETransactionStatus::AbortedByCascadingRetry:
-			return true;
+			break;
 	}
+	return false;
 }
 
 // The context status shows what state the AutoRTFM context is currently in.
