[AutoRTFM] Fix behavior of cascading aborts with non-materialized transactions.

Previously, we would treat a cascading abort as a regular abort if a deferred transaction existed.

#Synced-CL 50048130

diff --git a/Engine/Source/Programs/AutoRTFMTests/Private/API.cpp b/Engine/Source/Programs/AutoRTFMTests/Private/API.cpp
--- a/Engine/Source/Programs/AutoRTFMTests/Private/API.cpp
+++ b/Engine/Source/Programs/AutoRTFMTests/Private/API.cpp
@@ -1528,6 +1528,20 @@ TEST_CASE("API.OnCompleteIgnoredOutsideTransaction")
 	}
 }
 
+TEST_CASE("API.CascadingAbortTransaction_WithinUnscopedTransaction")
+{
+	AutoRTFM::Transact([]
+	{
+		AutoRTFM::TransactThenOpen([]
+		{
+			AutoRTFM::StartTransaction();
+			AutoRTFM::CascadingAbortTransaction();
+		});
+
+		FAIL("unreachable");
+	});
+}
+
 TEST_CASE("API.CascadingRetryTransaction")
 {
 	AUTORTFM_SCOPED_DISABLE_RETRY();
diff --git a/Engine/Source/Programs/AutoRTFMTests/Private/MetricsTests.cpp b/Engine/Source/Programs/AutoRTFMTests/Private/MetricsTests.cpp
--- a/Engine/Source/Programs/AutoRTFMTests/Private/MetricsTests.cpp
+++ b/Engine/Source/Programs/AutoRTFMTests/Private/MetricsTests.cpp
@@ -63,7 +63,7 @@ TEST_CASE("Metrics.TransactionCounts")
 		REQUIRE(Metrics.NumTransactionsStarted == 2);
 		REQUIRE(Metrics.NumTransactionsCommitted == 1);
 		REQUIRE(Metrics.NumTransactionsAborted == 1);
-		REQUIRE(Metrics.NumTransactionsAbortedByRequest == 0);
+		REQUIRE(Metrics.NumTransactionsAbortedByRequest == 1);
 		REQUIRE(Metrics.NumTransactionsAbortedByLanguage == 0);
 	}
 
diff --git a/Engine/Source/Runtime/AutoRTFM/Private/Context.cpp b/Engine/Source/Runtime/AutoRTFM/Private/Context.cpp
--- a/Engine/Source/Runtime/AutoRTFM/Private/Context.cpp
+++ b/Engine/Source/Runtime/AutoRTFM/Private/Context.cpp
@@ -113,15 +113,6 @@ void FContext::AbortTransaction(ETransactionStatus Reason)
 {
 	AUTORTFM_ASSERT(Status == EContextStatus::OnTrack);
 	AUTORTFM_ASSERT(HasAborted(Reason));
-
-	IncrementNumTransactionsAbortedMetric();
-
-	if (NumDeferredTransactions)
-	{
-		NumDeferredTransactions--;
-		return;
-	}
-
 	AUTORTFM_ASSERT(CurrentMaterializedTransaction);
 
 	switch (Reason)
@@ -143,22 +134,37 @@ void FContext::AbortTransaction(ETransactionStatus Reason)
 		// Set the status of all the call nests with the cascading failure.
 		SetCallNestTransactionCascadingStatus(Reason);
 
+		// Cancel out any deferred transactions.
+		IncrementNumTransactionsAbortedMetric(NumDeferredTransactions);
+		NumDeferredTransactions = 0;
+
 		// Immediately abort and pop the entire transaction stack.
 		while (CurrentMaterializedTransaction)
 		{
+			IncrementNumTransactionsAbortedMetric();
 			CurrentMaterializedTransaction->Abort();
 			PopTransaction();
 		}
 	}
 	else
 	{
-		// Set the status of all the call nests associated with the transaction
-		// with the cascading failure.
-		SetCallNestTransactionNonCascadingStatus(Reason);
+		if (NumDeferredTransactions > 0)
+		{
+			// The topmost transaction was never materialized, so an abort is trivial.
+			NumDeferredTransactions--;
+			IncrementNumTransactionsAbortedMetric();
+		}
+		else
+		{
+			// Set the status of all the call nests associated with the transaction
+			// with the (non-cascading) failure.
+			SetCallNestTransactionNonCascadingStatus(Reason);
 
-		// Abort and pop the top most transaction.
-		CurrentMaterializedTransaction->Abort();
-		PopTransaction();
+			// Abort and pop the top most transaction.
+			IncrementNumTransactionsAbortedMetric();
+			CurrentMaterializedTransaction->Abort();
+			PopTransaction();
+		}
 	}
 
 	Status = EContextStatus::OnTrack;
diff --git a/Engine/Source/Runtime/AutoRTFM/Private/MetricsPriv.h b/Engine/Source/Runtime/AutoRTFM/Private/MetricsPriv.h
--- a/Engine/Source/Runtime/AutoRTFM/Private/MetricsPriv.h
+++ b/Engine/Source/Runtime/AutoRTFM/Private/MetricsPriv.h
@@ -11,45 +11,45 @@ namespace AutoRTFM
 {
 
 /** Increments counter for number of top-level transactions started. */
-inline void IncrementNumTopLevelTransactionsStartedMetric()
+inline void IncrementNumTopLevelTransactionsStartedMetric(size_t Amount = 1)
 {
 	extern AutoRTFM::FAutoRTFMMetrics GAutoRTFMMetrics;
-	++GAutoRTFMMetrics.NumTopLevelTransactionsStarted;
+	GAutoRTFMMetrics.NumTopLevelTransactionsStarted += Amount;
 }
 
 /** Increments counter for number of transactions started. */
-inline void IncrementNumTransactionsStartedMetric()
+inline void IncrementNumTransactionsStartedMetric(size_t Amount = 1)
 {
 	extern AutoRTFM::FAutoRTFMMetrics GAutoRTFMMetrics;
-	++GAutoRTFMMetrics.NumTransactionsStarted;
+	GAutoRTFMMetrics.NumTransactionsStarted += Amount;
 }
 
 /** Increments counter for number of transactions committed. */
-inline void IncrementNumTransactionsCommittedMetric()
+inline void IncrementNumTransactionsCommittedMetric(size_t Amount = 1)
 {
 	extern AutoRTFM::FAutoRTFMMetrics GAutoRTFMMetrics;
-	++GAutoRTFMMetrics.NumTransactionsCommitted;
+	GAutoRTFMMetrics.NumTransactionsCommitted += Amount;
 }
 
 /** Increments counter for number of transactions aborted. */
-inline void IncrementNumTransactionsAbortedMetric()
+inline void IncrementNumTransactionsAbortedMetric(size_t Amount = 1)
 {
 	extern AutoRTFM::FAutoRTFMMetrics GAutoRTFMMetrics;
-	++GAutoRTFMMetrics.NumTransactionsAborted;
+	GAutoRTFMMetrics.NumTransactionsAborted += Amount;
 }
 
 /** Increments counter for number of transactions aborted via user request (AbortTransaction). */
-inline void IncrementNumTransactionsAbortedByRequestMetric()
+inline void IncrementNumTransactionsAbortedByRequestMetric(size_t Amount = 1)
 {
 	extern AutoRTFM::FAutoRTFMMetrics GAutoRTFMMetrics;
-	++GAutoRTFMMetrics.NumTransactionsAbortedByRequest;
+	GAutoRTFMMetrics.NumTransactionsAbortedByRequest += Amount;
 }
 
 /** Increments counter for number of transactions aborted due to violating transactionality (e.g. calling a non-closed function). */
-inline void IncrementNumTransactionsAbortedByLanguageMetric()
+inline void IncrementNumTransactionsAbortedByLanguageMetric(size_t Amount = 1)
 {
 	extern AutoRTFM::FAutoRTFMMetrics GAutoRTFMMetrics;
-	++GAutoRTFMMetrics.NumTransactionsAbortedByLanguage;
+	GAutoRTFMMetrics.NumTransactionsAbortedByLanguage += Amount;
 }
 
 /** Updates the peak memory usage metrics with the passed-in value, if it's higher than the current peak value. */
