[AutoRTFM] Replace autortfm_is_transactional() with autortfm_is_context_status(autortfm_context_status)

Add compiler constant-folding optimizations for calling autortfm_is_context_status() and autortfm_get_context_status() from the closed, so we don't lose any performance or increase binary size.

Allows us to replace the confusing, common expression `AutoRTFM::IsClosed() || AutoRTFM::IsCommittingOrAborting()` with `!AutoRTFM::IsIdle()`, and not incur any additional performance costs, or code bloat.

#Synced-CL 50245569

diff --git a/Engine/Source/Programs/AutoRTFMTests/Private/API.cpp b/Engine/Source/Programs/AutoRTFMTests/Private/API.cpp
--- a/Engine/Source/Programs/AutoRTFMTests/Private/API.cpp
+++ b/Engine/Source/Programs/AutoRTFMTests/Private/API.cpp
@@ -7,27 +7,6 @@
 #include "AutoRTFM/Testing.h"
 #include "Catch2Includes.h"
 
-TEST_CASE("API.autortfm_is_transactional")
-{
-    REQUIRE(false == autortfm_is_transactional());
-
-    bool bInTransaction = false;
-    bool bInOpenNest = false;
-
-    AutoRTFM::Testing::Commit([&]
-    {
-        bInTransaction = autortfm_is_transactional();
-
-        AutoRTFM::Open([&]
-        {
-            bInOpenNest = autortfm_is_transactional();
-        });
-    });
-
-    REQUIRE(true == bInTransaction);
-    REQUIRE(true == bInOpenNest);
-}
-
 TEST_CASE("API.autortfm_is_closed")
 {
     REQUIRE(false == autortfm_is_closed());
diff --git a/Engine/Source/Runtime/AutoRTFM/Private/API.cpp b/Engine/Source/Runtime/AutoRTFM/Private/API.cpp
--- a/Engine/Source/Runtime/AutoRTFM/Private/API.cpp
+++ b/Engine/Source/Runtime/AutoRTFM/Private/API.cpp
@@ -430,11 +430,6 @@ namespace
 {
 
 // Internal closed-variant implementations.
-AUTORTFM_INTERNAL bool RTFM_autortfm_is_transactional()
-{
-	return true;
-}
-
 AUTORTFM_INTERNAL autortfm_transaction_status RTFM_autortfm_transact(
 	void (*UninstrumentedWork)(void*), void (*InstrumentedWork)(void*), void* Arg)
 {
@@ -679,26 +674,6 @@ extern "C" void autortfm_shutdown() noexcept
 	FContext::Destroy();
 }
 
-// Each function will be forked by the compiler into an open and closed variant.
-// autortfm_is_closed() is used to branch to the closed variants declared above.
-extern "C" bool autortfm_is_transactional() noexcept
-{
-	if (autortfm_is_closed())
-	{
-		return RTFM_autortfm_is_transactional();
-	}
-
-	if (ForTheRuntime::IsAutoRTFMRuntimeEnabled())
-	{
-		if (FContext* Context = FContext::Get())
-		{
-			return Context->IsTransactional();
-		}
-	}
-
-	return false;
-}
-
 extern "C" UE_AUTORTFM_ALWAYS_OPEN autortfm_context_status autortfm_get_context_status() noexcept
 {
 	if (ForTheRuntime::IsAutoRTFMRuntimeEnabled())
@@ -712,6 +687,11 @@ extern "C" UE_AUTORTFM_ALWAYS_OPEN autortfm_context_status autortfm_get_context_
 	return autortfm_status_idle;
 }
 
+extern "C" UE_AUTORTFM_ALWAYS_OPEN bool autortfm_is_context_status(autortfm_context_status Status) noexcept
+{
+	return autortfm_get_context_status() == Status;
+}
+
 extern "C" autortfm_transaction_status autortfm_transact(
 	void (*UninstrumentedWork)(void*), void (*InstrumentedWork)(void*), void* Arg)
 {
diff --git a/Engine/Source/Runtime/AutoRTFM/Public/AutoRTFM.h b/Engine/Source/Runtime/AutoRTFM/Public/AutoRTFM.h
--- a/Engine/Source/Runtime/AutoRTFM/Public/AutoRTFM.h
+++ b/Engine/Source/Runtime/AutoRTFM/Public/AutoRTFM.h
@@ -389,14 +389,6 @@ UE_AUTORTFM_API void Shutdown();
 UE_AUTORTFM_CRITICAL_INLINE void Shutdown() {}
 #endif
 
-// Tells if we are currently running in a transaction. This will return true in an open nest
-// (see `Open`). This function is handled specially in the compiler and will be constant folded
-// as true in closed code, or preserved as a function call in open code.
-UE_AUTORTFM_CRITICAL_INLINE_ALWAYS bool IsTransactional()
-{
-	return autortfm_is_transactional();
-}
-
 // Tells if we are currently running in the closed nest of a transaction. By default,
 // transactional code is in a closed nest; the only way to be in an open nest is to request it
 // via `Open`. This function is handled specially in the compiler and will be constant folded
@@ -407,29 +399,39 @@ UE_AUTORTFM_CRITICAL_INLINE_ALWAYS bool IsClosed()
 }
 
 // Returns the current AutoRTFM context status.
+// Note: autortfm_get_context_status() is handled specially in the compiler and
+// will be constant folded as EContextStatus::OnTrack when called in the closed.
 UE_AUTORTFM_CRITICAL_INLINE EContextStatus ContextStatus()
 {
 	return static_cast<EContextStatus>(autortfm_get_context_status());
 }
 
-// Tells us if we are currently committing a transaction.
-// This will return true when inside an on-abort, on-commit, on-complete or on-retry handler.
-UE_AUTORTFM_CRITICAL_INLINE bool IsCommitting()
+// Returns true if we are currently running in a transaction, and outside of an
+// on-commit, on-abort, on-retry or on-complete handler.
+// This will constant folded by the compiler to true when called in the closed.
+UE_AUTORTFM_CRITICAL_INLINE_ALWAYS bool IsTransactional()
+{
+	return autortfm_is_context_status(autortfm_status_on_track);
+}
+
+// Returns true if we are currently committing a transaction.
+// This will return true when inside an on-commit handler.
+UE_AUTORTFM_CRITICAL_INLINE_ALWAYS bool IsCommitting()
 {
-	return ContextStatus() == EContextStatus::Committing;
+	return autortfm_is_context_status(autortfm_status_committing);
 }
 
-// Tells us if we are currently retrying a transaction.
+// Returns true if we are currently retrying a transaction.
 // This will return true when inside an on-retry handler.
-UE_AUTORTFM_CRITICAL_INLINE bool IsRetrying()
+UE_AUTORTFM_CRITICAL_INLINE_ALWAYS bool IsRetrying()
 {
-	return ContextStatus() == EContextStatus::Retrying;
+	return autortfm_is_context_status(autortfm_status_retrying);
 }
 
 // Returns true when inside an on-abort, on-commit, on-complete, on-retry handler
 // or inside a static initializer.
 // TODO(SOL-8967): This function name is misleading. Rename to something sensible.
-UE_AUTORTFM_CRITICAL_INLINE bool IsCommittingOrAborting()
+UE_AUTORTFM_CRITICAL_INLINE_ALWAYS bool IsCommittingOrAborting()
 {
 	switch (ContextStatus())
 	{
diff --git a/Engine/Source/Runtime/AutoRTFM/Public/AutoRTFM/CAPI.h b/Engine/Source/Runtime/AutoRTFM/Public/AutoRTFM/CAPI.h
--- a/Engine/Source/Runtime/AutoRTFM/Public/AutoRTFM/CAPI.h
+++ b/Engine/Source/Runtime/AutoRTFM/Public/AutoRTFM/CAPI.h
@@ -57,18 +57,6 @@ extern "C"
 		autortfm_transaction_aborted_by_cascading_retry,
 	} autortfm_transaction_status;
 
-	// This must match AutoRTFM::EContextStatus.
-	typedef enum
-	{
-		autortfm_status_idle = 0,
-		autortfm_status_on_track,
-		autortfm_status_committing,
-		autortfm_status_aborting,
-		autortfm_status_retrying,
-		autortfm_status_completing,
-		autortfm_status_in_static_local_initializer,
-	} autortfm_context_status;
-
 	// AutoRTFM logging severity.
 	typedef enum
 	{
@@ -223,20 +211,29 @@ UE_AUTORTFM_CRITICAL_INLINE bool autortfm_is_closed(void) AUTORTFM_NOEXCEPT
 #endif
 
 #if UE_AUTORTFM_ENABLED
-	UE_AUTORTFM_API bool autortfm_is_transactional(void) AUTORTFM_NOEXCEPT;
+	// Returns the current AutoRTFM context status.
+	// This function is handled specially in the compiler and will be constant
+	// folded as autortfm_status_on_track in closed code, or preserved as a
+	// function call in open code.
+	UE_AUTORTFM_API autortfm_context_status autortfm_get_context_status(void) AUTORTFM_NOEXCEPT;
 #else
-UE_AUTORTFM_CRITICAL_INLINE bool autortfm_is_transactional(void) AUTORTFM_NOEXCEPT
+UE_AUTORTFM_CRITICAL_INLINE autortfm_context_status autortfm_get_context_status(void) AUTORTFM_NOEXCEPT
 {
-	return false;
+	return autortfm_status_idle;
 }
 #endif
 
 #if UE_AUTORTFM_ENABLED
-	UE_AUTORTFM_API autortfm_context_status autortfm_get_context_status(void) AUTORTFM_NOEXCEPT;
+	// Returns true if the AutoRTFM context status is equal to Status.
+	// This is equivallent to (autortfm_get_context_status() == Status) and exists
+	// as a constant-folding optimization for non-optimized builds.
+	// This function is handled specially in the compiler and will be constant
+	// folded if possible.
+	UE_AUTORTFM_API bool autortfm_is_context_status(autortfm_context_status Status) AUTORTFM_NOEXCEPT;
 #else
-UE_AUTORTFM_CRITICAL_INLINE autortfm_context_status autortfm_get_context_status(void) AUTORTFM_NOEXCEPT
+UE_AUTORTFM_CRITICAL_INLINE bool autortfm_is_context_status(autortfm_context_status Status) AUTORTFM_NOEXCEPT
 {
-	return autortfm_status_idle;
+	return Status == autortfm_status_idle;
 }
 #endif
 
diff --git a/Engine/Source/Runtime/AutoRTFM/Public/AutoRTFM/Constants.h b/Engine/Source/Runtime/AutoRTFM/Public/AutoRTFM/Constants.h
--- a/Engine/Source/Runtime/AutoRTFM/Public/AutoRTFM/Constants.h
+++ b/Engine/Source/Runtime/AutoRTFM/Public/AutoRTFM/Constants.h
@@ -10,6 +10,18 @@
 extern "C"
 {
 #endif
+	// An enumerator of AutoRTFM context status modes.
+	// This must match AutoRTFM::EContextStatus.
+	typedef enum
+	{
+		autortfm_status_idle = 0,
+		autortfm_status_on_track,
+		autortfm_status_committing,
+		autortfm_status_aborting,
+		autortfm_status_retrying,
+		autortfm_status_completing,
+		autortfm_status_in_static_local_initializer,
+	} autortfm_context_status;
 
 	// An enumerator of AutoRTFM sanitizer modes.
 	// The AutoRTFM sanitizer is used to detect modification by open-code to memory that was written by
