[AutoRTFM] Make the sanitizer detect mixed open <-> closed writes from all threads.

Move all AutoRTFM sanitizer code to its own Sanitizer.[h|cpp] and make this concurrent-access safe.

Add AutoRTFM::Allocator, a C++ concept for a static interface to an allocator.
Add AutoRTFM::FAllocator which wraps the external API allocator functions.
Add AutoRTFM::Allocator template parameters to a number of the AutoRTFM container types, so they allocate with the given allocator.
Use a new OS-allocator for the sanitizer, to avoid deadlocks between mutexes in UE's allocators and the sanitizer lock.

Despite all the new mutex locking, this appears to be around 25% *faster* than the previous implementation, as we're no longer instrumenting the allocator open writes.

#Synced-CL 50319606

diff --git a/Engine/Source/Programs/AutoRTFMTests/Private/Allocators.cpp b/Engine/Source/Programs/AutoRTFMTests/Private/Allocators.cpp
--- a/Engine/Source/Programs/AutoRTFMTests/Private/Allocators.cpp
+++ b/Engine/Source/Programs/AutoRTFMTests/Private/Allocators.cpp
@@ -1,12 +1,11 @@
 // Copyright Epic Games, Inc. All Rights Reserved.
 
+#include "Allocator.h"
 #include "AutoRTFM.h"
 #include "AutoRTFM/Testing.h"
 #include "Catch2Includes.h"
-#include "ExternAPI.h"
 #include "MemoryStats.h"
 
-#include <new>
 #include <stdlib.h>
 
 #ifdef __clang__
@@ -375,23 +374,23 @@ TEST_CASE("Allocators.MemoryStats.GetTotalBytesAllocated")
 {
 	size_t StartingSize = AutoRTFM::FMemoryStats::GetTotalBytesAllocated();
 
-	void* const Allocate1MBytes = AutoRTFM::Allocate(1'000'000, /*Alignment*/ 1);
+	void* const Allocate1MBytes = AutoRTFM::FAllocator::Allocate(1'000'000, /*Alignment*/ 1);
 	size_t CurrentSize = AutoRTFM::FMemoryStats::GetTotalBytesAllocated();
 	REQUIRE(StartingSize + 1'000'000 == CurrentSize);
 
-	void* const Allocate49MBytes = AutoRTFM::AllocateZeroed(49'000'000, /*Alignment*/ 1);
+	void* const Allocate49MBytes = AutoRTFM::FAllocator::AllocateZeroed(49'000'000, /*Alignment*/ 1);
 	CurrentSize = AutoRTFM::FMemoryStats::GetTotalBytesAllocated();
 	REQUIRE(StartingSize + 50'000'000 == CurrentSize);
 
-	AutoRTFM::Free(Allocate1MBytes, 1'000'000);
+	AutoRTFM::FAllocator::Free(Allocate1MBytes, 1'000'000);
 	CurrentSize = AutoRTFM::FMemoryStats::GetTotalBytesAllocated();
 	REQUIRE(StartingSize + 49'000'000 == CurrentSize);
 
-	void* const Allocate7MBytes = AutoRTFM::Reallocate(Allocate49MBytes, 7'000'000, /*Alignment*/ 1, /*PreviousSize*/ 49'000'000);
+	void* const Allocate7MBytes = AutoRTFM::FAllocator::Reallocate(Allocate49MBytes, 7'000'000, /*Alignment*/ 1, /*PreviousSize*/ 49'000'000);
 	CurrentSize = AutoRTFM::FMemoryStats::GetTotalBytesAllocated();
 	REQUIRE(StartingSize + 7'000'000 == CurrentSize);
 
-	AutoRTFM::Free(Allocate7MBytes, 7'000'000);
+	AutoRTFM::FAllocator::Free(Allocate7MBytes, 7'000'000);
 	CurrentSize = AutoRTFM::FMemoryStats::GetTotalBytesAllocated();
 	REQUIRE(CurrentSize == StartingSize);
 }
diff --git a/Engine/Source/Programs/AutoRTFMTests/Private/CallstackTreeTests.cpp b/Engine/Source/Programs/AutoRTFMTests/Private/CallstackTreeTests.cpp
--- a/Engine/Source/Programs/AutoRTFMTests/Private/CallstackTreeTests.cpp
+++ b/Engine/Source/Programs/AutoRTFMTests/Private/CallstackTreeTests.cpp
@@ -5,19 +5,21 @@
 
 TEST_CASE("CallstackTree")
 {
-	using FHandle = AutoRTFM::FCallstackTree::FHandle;
-	static constexpr AutoRTFM::FCallstackTree::FHandle InvalidHandle = AutoRTFM::FCallstackTree::InvalidHandle;
+	using FCallstackTree = AutoRTFM::TCallstackTree<>;
+
+	using FHandle = FCallstackTree::FHandle;
+	static constexpr FCallstackTree::FHandle InvalidHandle = FCallstackTree::InvalidHandle;
 
 	SECTION("Empty")
 	{
-		AutoRTFM::FCallstackTree Tree;
+		FCallstackTree Tree;
 		REQUIRE(Tree.Add(0, nullptr) == InvalidHandle);
 		REQUIRE((Tree.Get(InvalidHandle).IsEmpty()));
 	}
 
 	SECTION("Shallow, Unique")
 	{
-		AutoRTFM::FCallstackTree Tree;
+		FCallstackTree Tree;
 		AutoRTFM::TStack<autortfm_callstack_frame, 1> A{1};
 		FHandle HandleA = Tree.Add(A.Num(), &A.Front());
 		REQUIRE(HandleA == 1);
@@ -37,7 +39,7 @@ TEST_CASE("CallstackTree")
 
 	SECTION("Shallow, Duplicates")
 	{
-		AutoRTFM::FCallstackTree Tree;
+		FCallstackTree Tree;
 		AutoRTFM::TStack<autortfm_callstack_frame, 1> A{1};
 		FHandle HandleA = Tree.Add(A.Num(), &A.Front());
 		REQUIRE(HandleA == 1);
@@ -57,7 +59,7 @@ TEST_CASE("CallstackTree")
 	
 	SECTION("Deep")
 	{
-		AutoRTFM::FCallstackTree Tree;
+		FCallstackTree Tree;
 		AutoRTFM::TStack<autortfm_callstack_frame, 5> A{5, 4, 3, 2, 1};
 		FHandle HandleA = Tree.Add(A.Num(), &A.Front());
 		REQUIRE(HandleA == 5);
diff --git a/Engine/Source/Programs/AutoRTFMTests/Private/PoolTests.cpp b/Engine/Source/Programs/AutoRTFMTests/Private/PoolTests.cpp
--- a/Engine/Source/Programs/AutoRTFMTests/Private/PoolTests.cpp
+++ b/Engine/Source/Programs/AutoRTFMTests/Private/PoolTests.cpp
@@ -1,11 +1,12 @@
 // Copyright Epic Games, Inc. All Rights Reserved.
 
+#include "AutoRTFM.h"
 #include "Pool.h"
 #include "Catch2Includes.h"
 #include "ObjectLifetimeHelper.h"
+
 #include <algorithm>
 #include <vector>
-#include <random>
 
 using TrivialPool4 = AutoRTFM::TPool<int, 4>;
 using TrivialPool64 = AutoRTFM::TPool<int, 64>;
@@ -14,6 +15,8 @@ using NonTrivialPool64 = AutoRTFM::TPool<AutoRTFMTestUtils::FObjectLifetimeHelpe
 
 TEMPLATE_TEST_CASE("Pool", "",  TrivialPool4, TrivialPool64, NonTrivialPool4, NonTrivialPool64)
 {
+	AutoRTFM::UnreachableIfClosed();
+
 	Catch::SimplePcg32 Rand(0x1234);
 	TestType Pool;
 	using TItem = typename TestType::FItem;
diff --git a/Engine/Source/Runtime/AutoRTFM/Private/API.cpp b/Engine/Source/Runtime/AutoRTFM/Private/API.cpp
--- a/Engine/Source/Runtime/AutoRTFM/Private/API.cpp
+++ b/Engine/Source/Runtime/AutoRTFM/Private/API.cpp
@@ -31,9 +31,6 @@ int GAutoRTFMRetryTransactions = AutoRTFM::ForTheRuntime::EAutoRTFMRetryTransact
 // See `AutoRTFM::ForTheRuntime::CoinTossDisable` for implementation details.
 float GAutoRTFMEnabledProbability = 100.0f;
 
-AUTORTFM_SANITIZER_ONLY(AutoRTFM::Sanitizer::EMode GSanitizerMode = AutoRTFM::Sanitizer::EMode::Error);
-AUTORTFM_SANITIZER_ONLY(bool GSanitizerRecordClosedWriteCallstacks = false);
-
 bool GAutoRTFMEnsureOnInternalAbort = true;
 
 // A linked-list of open->closed function tables populated by
@@ -109,67 +106,6 @@ AUTORTFM_INTERNAL bool IsAutoRTFMInitialized()
 namespace AutoRTFM
 {
 
-#if AUTORTFM_SANITIZER
-namespace Sanitizer
-{
-EMode GetMode()
-{
-	return GSanitizerMode;
-}
-
-void SetMode(EMode Mode)
-{
-	if (GSanitizerMode != Mode)
-	{
-		GSanitizerMode = Mode;
-		if (GExternAPI.OnSanitizerModeChanged)
-		{
-			GExternAPI.OnSanitizerModeChanged();
-		}
-	}
-}
-
-void BeginDisable()
-{
-	if (FContext* Context = FContext::Get())
-	{
-		if (FTransaction* Transaction = Context->GetCurrentTransaction())
-		{
-			Transaction->DisableSanitizer();
-		}
-	}
-}
-
-void EndDisable()
-{
-	if (FContext* Context = FContext::Get())
-	{
-		if (FTransaction* Transaction = Context->GetCurrentTransaction())
-		{
-			Transaction->EnableSanitizer();
-		}
-	}
-}
-
-bool RecordClosedWriteCallstacks()
-{
-	return GSanitizerRecordClosedWriteCallstacks;
-}
-
-void SetRecordClosedWriteCallstacks(bool bEnable)
-{
-	if (GSanitizerRecordClosedWriteCallstacks != bEnable)
-	{
-		GSanitizerRecordClosedWriteCallstacks = bEnable;
-		if (GExternAPI.OnSanitizerRecordClosedWriteCallstacksChanged)
-		{
-			GExternAPI.OnSanitizerRecordClosedWriteCallstacksChanged();
-		}
-	}
-}
-}  // namespace Sanitizer
-#endif  // AUTORTFM_SANITIZER
-
 namespace ForTheRuntime
 {
 bool SetAutoRTFMRuntime(EAutoRTFMEnabledState State)
@@ -467,7 +403,7 @@ extern "C" UE_AUTORTFM_NOAUTORTFM UE_AUTORTFM_API void autortfm_pre_open(bool bN
 #if AUTORTFM_SANITIZER
 	if (bNoSanitize)
 	{
-		FContext::Get()->GetCurrentTransaction()->DisableSanitizer();
+		Sanitizer::BeginDisable();
 	}
 #endif  // AUTORTFM_SANITIZER
 }
@@ -485,7 +421,7 @@ extern "C" UE_AUTORTFM_NOAUTORTFM UE_AUTORTFM_API void autortfm_post_open(bool b
 		// if the current transaction at autortfm_pre_open() was aborted and
 		// popped, in which case the current transaction at autortfm_post_open()
 		// will now be an ancestor transaction.
-		Context->GetCurrentTransaction()->EnableSanitizer();
+		Sanitizer::EndDisable();
 	}
 #endif  // AUTORTFM_SANITIZER
 }
@@ -666,12 +602,15 @@ extern "C" void autortfm_initialize(const autortfm_extern_api* ExternAPI) noexce
 	GExternAPI = *ExternAPI;
 	FContext::Create();
 	ProcessAllPendingOpenToClosedRegistrations();
+	AUTORTFM_SANITIZER_ONLY(Sanitizer::Initialize());
 }
 
 extern "C" void autortfm_shutdown() noexcept
 {
 	AutoRTFM::UnreachableIfClosed("TODO: Mark this as AUTORTFM_DISABLED and fixup callers");
+
 	FContext::Destroy();
+	AUTORTFM_SANITIZER_ONLY(Sanitizer::Shutdown());
 }
 
 extern "C" UE_AUTORTFM_ALWAYS_OPEN autortfm_context_status autortfm_get_context_status() noexcept
@@ -793,13 +732,17 @@ extern "C" AUTORTFM_INTERNAL UE_AUTORTFM_API void autortfm_open_from_closed_no_s
 {
 	FContext* const Context = FContext::Get();
 
-	AUTORTFM_SANITIZER_ONLY(Context->GetCurrentTransaction()->DisableSanitizer());
+	AUTORTFM_SANITIZER_ONLY(Sanitizer::BeginDisable());
 
 	Work(Arg);
 
-	AUTORTFM_SANITIZER_ONLY(Context->GetCurrentTransaction()->EnableSanitizer());
-
 	Context->MaybeThrow();
+
+	// Note: This must come *after* MaybeThrow(), because we'll only throw
+	// if the current transaction at autortfm_pre_open() was aborted and
+	// popped, in which case the current transaction at autortfm_post_open()
+	// will now be an ancestor transaction.
+	AUTORTFM_SANITIZER_ONLY(Sanitizer::EndDisable());
 }
 
 // Note: Internal.aem maps this to autortfm_open_from_closed_no_sanitize when called in the closed
diff --git a/Engine/Source/Runtime/AutoRTFM/Private/Allocator.h b/Engine/Source/Runtime/AutoRTFM/Private/Allocator.h
new file mode 100644
--- /dev/null
+++ b/Engine/Source/Runtime/AutoRTFM/Private/Allocator.h
@@ -0,0 +1,101 @@
+// Copyright Epic Games, Inc. All Rights Reserved.
+
+#pragma once
+
+#if (defined(__AUTORTFM) && __AUTORTFM)
+
+#include "ExternAPI.h"
+#include "MemoryStats.h"
+
+#include <concepts>
+#include <utility>
+
+namespace AutoRTFM
+{
+
+// A collection of memory allocator static methods which can be passed by
+// template argument to AutoRTFM's container types, ensuring that the container
+// uses these methods for all heap allocations.
+template <typename T>
+concept Allocator = requires {
+	{
+		// The function used to allocate memory from the heap. Must not be null.
+		// This API should never return null; if the allocator cannot satisfy the request, it must abort the program.
+		// static void* Allocate(size_t Size, size_t Alignment)
+		T::Allocate(std::declval<size_t>(), std::declval<size_t>())
+	} -> std::same_as<void*>;
+	{
+
+		// The function used to reallocate memory from the heap. Must not be null.
+		// This API should never return null; if the allocator cannot satisfy the request, it must abort the program.
+		// static void* Reallocate(void* Pointer, size_t Size, size_t Alignment, size_t PreviousSize)
+		T::Reallocate(std::declval<void*>(), std::declval<size_t>(), std::declval<size_t>(), std::declval<size_t>())
+	} -> std::same_as<void*>;
+	{
+
+		// The function used to allocate zeroed memory from the heap. Must not be null.
+		// This API should never return null; if the allocator cannot satisfy the request, it must abort the program.
+		// static void* AllocateZeroed(size_t Size, size_t Alignment)
+		T::AllocateZeroed(std::declval<size_t>(), std::declval<size_t>())
+	} -> std::same_as<void*>;
+	{
+		// The function used to free memory allocated by Allocate(), Reallocate() or AllocateZeroed().
+		// Must not be null.
+		// static void Free(void* Pointer, size_t AllocationSize)
+		T::Free(std::declval<void*>(), std::declval<size_t>())
+	} -> std::same_as<void>;
+};
+
+// The AutoRTFM default Allocator that wraps the external API's memory
+// allocation functions.
+struct AUTORTFM_DISABLE_NO_SANITIZE FAllocator
+{
+	static inline void* Allocate(size_t Size, size_t Alignment)
+	{
+		FMemoryStats::UpdateBytesAllocated(Size);
+		return GExternAPI.Allocate(Size, Alignment);
+	}
+
+	static inline void* Reallocate(void* Pointer, size_t Size, size_t Alignment, size_t PreviousSize)
+	{
+		FMemoryStats::UpdateBytesAllocated(static_cast<ptrdiff_t>(Size) - static_cast<ptrdiff_t>(PreviousSize));
+		return GExternAPI.Reallocate(Pointer, Size, Alignment, PreviousSize);
+	}
+
+	static inline void* AllocateZeroed(size_t Size, size_t Alignment)
+	{
+		FMemoryStats::UpdateBytesAllocated(Size);
+		return GExternAPI.AllocateZeroed(Size, Alignment);
+	}
+
+	static inline void Free(void* Pointer, size_t AllocationSize)
+	{
+		FMemoryStats::UpdateBytesAllocated(-static_cast<ptrdiff_t>(AllocationSize));
+		return GExternAPI.Free(Pointer, AllocationSize);
+	}
+};
+
+static_assert(Allocator<FAllocator>);
+
+// Constructs a new T with the given arguments, using Allocator.
+template <typename Allocator, typename T, typename... ARGS>
+static inline T* New(ARGS&&... Args)
+{
+	return new (Allocator::Allocate(sizeof(T), alignof(T))) T(std::forward<ARGS>(Args)...);
+}
+
+// Destructs and frees a T* allocated with New<Allocator, T>(...).
+// Ptr should be immediately set to null after calling.
+template <typename Allocator, typename T, typename... ARGS>
+static inline void Delete(T* Ptr)
+{
+	if (Ptr)
+	{
+		Ptr->~T();
+		Allocator::Free(Ptr, sizeof(T));
+	}
+}
+
+}  // namespace AutoRTFM
+
+#endif  // (defined(__AUTORTFM) && __AUTORTFM)
diff --git a/Engine/Source/Runtime/AutoRTFM/Private/BlockAllocator.h b/Engine/Source/Runtime/AutoRTFM/Private/BlockAllocator.h
--- a/Engine/Source/Runtime/AutoRTFM/Private/BlockAllocator.h
+++ b/Engine/Source/Runtime/AutoRTFM/Private/BlockAllocator.h
@@ -4,7 +4,7 @@
 
 #if (defined(__AUTORTFM) && __AUTORTFM)
 
-#include "ExternAPI.h"
+#include "Allocator.h"
 #include "Utils.h"
 
 #include <cstddef>
@@ -22,7 +22,7 @@ namespace AutoRTFM
 //   DataAlignment - The maximum supported alignment of data from this allocator.
 //   GrowthPercentage - The percentage to grow the block size each reallocation.
 template <size_t InlineBlockDataSize = 256, size_t DataAlignment = 16, size_t GrowthPercentage = 200>
-class TBlockAllocator
+class AUTORTFM_DISABLE_NO_SANITIZE TBlockAllocator
 {
 public:
 	// Constructor
@@ -40,7 +40,7 @@ public:
 		while (Tail != &InlineBlock)
 		{
 			FBlockHeader* Prev = Tail->Prev;
-			AutoRTFM::Free(Tail, Tail->TotalSize());
+			FAllocator::Free(Tail, Tail->TotalSize());
 			Tail = Prev;
 		}
 		InlineBlock.Remaining = InlineBlockDataSize;
@@ -68,14 +68,6 @@ public:
 		return Allocation;
 	}
 
-	// Constructs and returns a new T into the memory returned by calling
-	// Allocate(sizeof(T), alignof(T)).
-	template <typename T, typename... ArgTypes>
-	T* New(ArgTypes&&... Args)
-	{
-		return new (Allocate(sizeof(T), alignof(T))) T(Forward<ArgTypes>(Args)...);
-	}
-
 private:
 	TBlockAllocator(TBlockAllocator&&) = delete;
 	TBlockAllocator(const TBlockAllocator&) = delete;
@@ -97,7 +89,7 @@ private:
 		// Allocates, constructs and returns a new block from the heap-allocated memory.
 		static FBlockHeader* New(FBlockHeader* Prev, size_t BlockDataSize)
 		{
-			void* Memory = AutoRTFM::Allocate(sizeof(FBlockHeader) + BlockDataSize, BlockAlignment);
+			void* Memory = FAllocator::Allocate(sizeof(FBlockHeader) + BlockDataSize, BlockAlignment);
 			return new (Memory) FBlockHeader{Prev, BlockDataSize};
 		}
 
diff --git a/Engine/Source/Runtime/AutoRTFM/Private/CallstackTree.h b/Engine/Source/Runtime/AutoRTFM/Private/CallstackTree.h
--- a/Engine/Source/Runtime/AutoRTFM/Private/CallstackTree.h
+++ b/Engine/Source/Runtime/AutoRTFM/Private/CallstackTree.h
@@ -4,6 +4,7 @@
 
 #if (defined(__AUTORTFM) && __AUTORTFM)
 
+#include "Allocator.h"
 #include "AutoRTFM/CAPI.h"
 #include "AutoRTFM/Defines.h"
 #include "Stack.h"
@@ -18,7 +19,8 @@ namespace AutoRTFM
 
 // FCallstackTree is a store for callstacks that internally holds a DAG of
 // frames which reduces memory consumption by de-duplicating common base frames.
-class AUTORTFM_DISABLE_NO_SANITIZE FCallstackTree final
+template <Allocator Allocator = FAllocator>
+class AUTORTFM_DISABLE_NO_SANITIZE TCallstackTree final
 {
 public:
 	using FHandle = uint32_t;
@@ -54,9 +56,9 @@ public:
 	}
 
 	// Retrieves a callstack previously added with Add().
-	TStack<autortfm_callstack_frame, 8> Get(FHandle Handle) const
+	TStack<autortfm_callstack_frame, 8, Allocator> Get(FHandle Handle) const
 	{
-		TStack<autortfm_callstack_frame, 8> StackTrace;
+		TStack<autortfm_callstack_frame, 8, Allocator> StackTrace;
 		while (Handle != InvalidHandle)
 		{
 			FEntry const& Entry = Entries[Handle - 1];
@@ -75,7 +77,7 @@ private:
 		bool operator==(const FEntry&) const = default;
 	};
 
-	using FEntryList = TStack<FEntry, 32>;
+	using FEntryList = TStack<FEntry, 32, Allocator>;
 
 	struct FHasher
 	{
@@ -111,7 +113,7 @@ private:
 	};
 
 	FEntryList Entries;
-	std::unordered_set<FHandle, FHasher, FEquality, StlAllocator<FHandle>> Lookup{
+	std::unordered_set<FHandle, FHasher, FEquality, StlAllocator<FHandle, Allocator>> Lookup{
 		/* bucket_count */ 32,
 		/* hash */ FHasher{Entries},
 		/* equal */ FEquality{Entries},
diff --git a/Engine/Source/Runtime/AutoRTFM/Private/CompilerInterface.cpp b/Engine/Source/Runtime/AutoRTFM/Private/CompilerInterface.cpp
--- a/Engine/Source/Runtime/AutoRTFM/Private/CompilerInterface.cpp
+++ b/Engine/Source/Runtime/AutoRTFM/Private/CompilerInterface.cpp
@@ -437,25 +437,6 @@ extern "C" AUTORTFM_DISABLE_NO_SANITIZE UE_AUTORTFM_FORCENOINLINE UE_AUTORTFM_AP
 	AUTORTFM_FATAL("inlined AUTORTFM_DISABLE function called from the closed");
 }
 
-#if AUTORTFM_SANITIZER
-static bool bInAutoRTFMSanitizerOpenWrite = false;
-extern "C" AUTORTFM_DISABLE_NO_SANITIZE UE_AUTORTFM_FORCENOINLINE UE_AUTORTFM_API void autortfm_sanitizer_open_write(
-	void* Ptr, size_t Size) noexcept
-{
-	if (FContext* Context = FContext::Get(); Context && Context->IsContextThread())
-	{
-		if (!bInAutoRTFMSanitizerOpenWrite)
-		{
-			TScopedGuard RecursionGuard(bInAutoRTFMSanitizerOpenWrite, true);
-			if (FTransaction* Transaction = Context->GetMaterializedTransaction())
-			{
-				Transaction->SanitizerOpenWrite(Ptr, Size);
-			}
-		}
-	}
-}
-#endif
-
 }  // namespace AutoRTFM
 
 #endif  // defined(__AUTORTFM) && __AUTORTFM
diff --git a/Engine/Source/Runtime/AutoRTFM/Private/Context.cpp b/Engine/Source/Runtime/AutoRTFM/Private/Context.cpp
--- a/Engine/Source/Runtime/AutoRTFM/Private/Context.cpp
+++ b/Engine/Source/Runtime/AutoRTFM/Private/Context.cpp
@@ -5,10 +5,10 @@
 #include "Context.h"
 #include "ContextInlines.h"
 
+#include "Allocator.h"
 #include "AutoRTFM.h"
 #include "CallNest.h"
 #include "CallNestInlines.h"
-#include "ExternAPI.h"
 #include "MetricsPriv.h"
 #include "StackRange.h"
 #include "Stats.h"
@@ -29,7 +29,7 @@ FContext* FContext::Instance = nullptr;
 FContext* FContext::Create()
 {
 	AUTORTFM_ENSURE(Instance == nullptr);
-	void* Memory = AutoRTFM::Allocate(sizeof(FContext), alignof(FContext));
+	void* Memory = FAllocator::Allocate(sizeof(FContext), alignof(FContext));
 	Instance = new (Memory) FContext();
 	return Instance;
 }
@@ -42,7 +42,7 @@ void FContext::Destroy()
 		// attempt to use the context.
 		Instance = nullptr;
 		Context->~FContext();
-		AutoRTFM::Free(Context, sizeof(FContext));
+		FAllocator::Free(Context, sizeof(FContext));
 	}
 	else
 	{
@@ -53,7 +53,7 @@ void FContext::Destroy()
 FContext::~FContext()
 {
 #if AUTORTFM_SANITIZER
-	if (NumSanitizerIssues)
+	if (size_t NumSanitizerIssues = Sanitizer::NumberOfIssuesFound())
 	{
 		AUTORTFM_WARN("AutoRTFMSan: %zu sanitizer issues detected", NumSanitizerIssues);
 	}
@@ -542,43 +542,6 @@ void FContext::Throw()
 	CurrentNest->AbortJump.Throw();
 }
 
-#if AUTORTFM_SANITIZER
-FCallstackTree::FHandle FContext::GetCallstackHandle(size_t FramesToSkip /* = 0 */)
-{
-	static constexpr size_t MaxCallstack = 128;
-	autortfm_callstack_frame Callstack[MaxCallstack];
-	size_t const NumCaptured = GExternAPI.CaptureCallstack(MaxCallstack, Callstack);
-
-	// Attempt to trim off callstack frames under this function.
-	if (autortfm_callstack_frame ReturnAddress = reinterpret_cast<autortfm_callstack_frame>(__builtin_return_address(0)))
-	{
-		for (size_t I = 0, N = std::min(NumCaptured, static_cast<size_t>(8)); I < N; I++)
-		{
-			if (ReturnAddress == Callstack[I])
-			{
-				FramesToSkip += I;
-				break;
-			}
-		}
-	}
-
-	if (NumCaptured <= FramesToSkip)
-	{
-		return FCallstackTree::InvalidHandle;
-	}
-
-	return CallstackTree.Add(NumCaptured - FramesToSkip, Callstack + FramesToSkip);
-}
-
-void FContext::LogCallstack(FCallstackTree::FHandle Handle, autortfm_log_severity Severity)
-{
-	if (auto Callstack = CallstackTree.Get(Handle); !Callstack.IsEmpty())
-	{
-		GExternAPI.LogCallstack(Severity, Callstack.Num(), &Callstack.Front());
-	}
-}
-#endif
-
 }  // namespace AutoRTFM
 
 #endif  // defined(__AUTORTFM) && __AUTORTFM
diff --git a/Engine/Source/Runtime/AutoRTFM/Private/Context.h b/Engine/Source/Runtime/AutoRTFM/Private/Context.h
--- a/Engine/Source/Runtime/AutoRTFM/Private/Context.h
+++ b/Engine/Source/Runtime/AutoRTFM/Private/Context.h
@@ -15,10 +15,6 @@
 #include "Transaction.h"
 #include "Utils.h"
 
-#if AUTORTFM_SANITIZER
-#include "CallstackTree.h"
-#endif
-
 namespace AutoRTFM
 {
 
@@ -187,15 +183,6 @@ public:
 	// Returns the stack address limits for the caller thread.
 	static FStackRange GetThreadStackRange();
 
-#if AUTORTFM_SANITIZER
-	FCallstackTree::FHandle GetCallstackHandle(size_t FramesToSkip = 0);
-	void LogCallstack(FCallstackTree::FHandle, autortfm_log_severity);
-	void IncrementNumSanitizerIssues()
-	{
-		NumSanitizerIssues++;
-	}
-#endif
-
 private:
 	struct FAutoCommitOnException;
 
@@ -246,11 +233,6 @@ private:
 
 	TTaskArray<TTask<void()>> CompletionTasks{TaskPool};
 	TTaskArray<TTask<void()>> RetryTasks{TaskPool};
-
-#if AUTORTFM_SANITIZER
-	FCallstackTree CallstackTree;
-	size_t NumSanitizerIssues = 0;
-#endif
 };
 
 }  // namespace AutoRTFM
diff --git a/Engine/Source/Runtime/AutoRTFM/Private/ExternAPI.h b/Engine/Source/Runtime/AutoRTFM/Private/ExternAPI.h
--- a/Engine/Source/Runtime/AutoRTFM/Private/ExternAPI.h
+++ b/Engine/Source/Runtime/AutoRTFM/Private/ExternAPI.h
@@ -5,37 +5,12 @@
 #if (defined(__AUTORTFM) && __AUTORTFM)
 
 #include "AutoRTFM/CAPI.h"
-#include "MemoryStats.h"
 
 namespace AutoRTFM
 {
 
 UE_AUTORTFM_API extern autortfm_extern_api GExternAPI;
 
-inline void* Allocate(size_t Size, size_t Alignment)
-{
-	FMemoryStats::UpdateBytesAllocated(Size);
-	return GExternAPI.Allocate(Size, Alignment);
-}
-
-inline void* Reallocate(void* Pointer, size_t Size, size_t Alignment, size_t PreviousSize)
-{
-	FMemoryStats::UpdateBytesAllocated(static_cast<ptrdiff_t>(Size) - static_cast<ptrdiff_t>(PreviousSize));
-	return GExternAPI.Reallocate(Pointer, Size, Alignment, PreviousSize);
-}
-
-inline void* AllocateZeroed(size_t Size, size_t Alignment)
-{
-	FMemoryStats::UpdateBytesAllocated(Size);
-	return GExternAPI.AllocateZeroed(Size, Alignment);
-}
-
-inline void Free(void* Pointer, size_t AllocationSize)
-{
-	FMemoryStats::UpdateBytesAllocated(-static_cast<ptrdiff_t>(AllocationSize));
-	return GExternAPI.Free(Pointer, AllocationSize);
-}
-
 }  // namespace AutoRTFM
 
 #endif  // (defined(__AUTORTFM) && __AUTORTFM)
diff --git a/Engine/Source/Runtime/AutoRTFM/Private/HashMap.h b/Engine/Source/Runtime/AutoRTFM/Private/HashMap.h
--- a/Engine/Source/Runtime/AutoRTFM/Private/HashMap.h
+++ b/Engine/Source/Runtime/AutoRTFM/Private/HashMap.h
@@ -4,6 +4,7 @@
 
 #if (defined(__AUTORTFM) && __AUTORTFM)
 
+#include "Allocator.h"
 #include "StlAllocator.h"
 
 #include <functional>
@@ -24,10 +25,12 @@ struct TKeyAndValue
 // HashMap is an unordered hashmap.
 // TODO(SOL-7652): This currently wraps a std::unordered_map. Reimplement to
 // improve performance and to avoid platform-specific variations in behavior.
-template <typename KeyType, typename ValueType, typename Hash = std::hash<KeyType>, typename Equal = std::equal_to<KeyType>>
+template <typename KeyType, typename ValueType, typename Hash = std::hash<KeyType>, typename Equal = std::equal_to<KeyType>,
+	Allocator Allocator = FAllocator>
 class THashMap
 {
-	using FInnerMap = std::unordered_map<KeyType, ValueType, Hash, Equal, StlAllocator<std::pair<const KeyType, ValueType>>>;
+	using FInnerMap =
+		std::unordered_map<KeyType, ValueType, Hash, Equal, StlAllocator<std::pair<const KeyType, ValueType>, Allocator>>;
 
 	template <bool bConst>
 	class TIterator
diff --git a/Engine/Source/Runtime/AutoRTFM/Private/HitSet.h b/Engine/Source/Runtime/AutoRTFM/Private/HitSet.h
--- a/Engine/Source/Runtime/AutoRTFM/Private/HitSet.h
+++ b/Engine/Source/Runtime/AutoRTFM/Private/HitSet.h
@@ -4,7 +4,7 @@
 
 #if (defined(__AUTORTFM) && __AUTORTFM)
 
-#include "ExternAPI.h"
+#include "Allocator.h"
 #include "Utils.h"
 
 #include <cstdint>
@@ -163,7 +163,7 @@ public:
 			if (Payload != SmallPayload)
 			{
 				size_t const PayloadSize = Capacity() * sizeof(uintptr_t);
-				AutoRTFM::Free(Payload, PayloadSize);
+				FAllocator::Free(Payload, PayloadSize);
 				Payload = SmallPayload;
 			}
 			memset(SmallPayload, 0, sizeof(SmallPayload));
@@ -284,7 +284,7 @@ private:
 			IncreaseCapacity();
 
 			size_t const NewPayloadSize = Capacity() * sizeof(uintptr_t);
-			Payload = static_cast<uintptr_t*>(AutoRTFM::AllocateZeroed(NewPayloadSize, alignof(uintptr_t)));
+			Payload = static_cast<uintptr_t*>(FAllocator::AllocateZeroed(NewPayloadSize, alignof(uintptr_t)));
 			AUTORTFM_ASSERT(nullptr != Payload);
 
 			// Now we need to rehash and reinsert all the items.
@@ -311,7 +311,7 @@ private:
 
 			if (bNeedAnotherResize)
 			{
-				AutoRTFM::Free(Payload, NewPayloadSize);
+				FAllocator::Free(Payload, NewPayloadSize);
 				continue;
 			}
 
@@ -322,7 +322,7 @@ private:
 
 		if (OldPayload != SmallPayload)
 		{
-			AutoRTFM::Free(OldPayload, OldPayloadSize);
+			FAllocator::Free(OldPayload, OldPayloadSize);
 		}
 
 		return true;
diff --git a/Engine/Source/Runtime/AutoRTFM/Private/IntervalTree.h b/Engine/Source/Runtime/AutoRTFM/Private/IntervalTree.h
--- a/Engine/Source/Runtime/AutoRTFM/Private/IntervalTree.h
+++ b/Engine/Source/Runtime/AutoRTFM/Private/IntervalTree.h
@@ -4,14 +4,13 @@
 
 #if (defined(__AUTORTFM) && __AUTORTFM)
 
+#include "Allocator.h"
 #include "ContainerValidation.h"
-#include "ExternAPI.h"
 #include "Stack.h"
 #include "Utils.h"
 
 #include <cstring>
 #include <type_traits>
-#include <variant>
 
 namespace AutoRTFM
 {
@@ -39,7 +38,7 @@ struct FInterval final
 	uintptr_t End;
 };
 
-template <typename DataTypeIn>
+template <typename DataTypeIn, Allocator Allocator = FAllocator>
 struct AUTORTFM_DISABLE_NO_SANITIZE TIntervalTree final
 {
 private:
@@ -109,7 +108,7 @@ public:
 			return;
 		}
 
-		TStack<NodeIndex, 32, EContainerValidation::Disabled> ToProcess;
+		TStack<NodeIndex, 32, Allocator, EContainerValidation::Disabled> ToProcess;
 
 		ToProcess.Push(Other.Root);
 
@@ -266,7 +265,7 @@ private:
 			Count = 0;
 			Capacity = 0;
 
-			AutoRTFM::Free(Memory, MemorySizeInBytes);
+			Allocator::Free(Memory, MemorySizeInBytes);
 			Memory = nullptr;
 			MemorySizeInBytes = 0;
 		}
@@ -373,7 +372,7 @@ private:
 			size_t const ParentOffset = SubAllocate<NodeIndex>(TotalSize, NewCapacity);
 			size_t const ColorOffset = SubAllocate<ColorWord>(TotalSize, NumColorWords(NewCapacity));
 			size_t const DataOffset = bHasData ? SubAllocate<DataType>(TotalSize, NewCapacity) : 0;
-			auto* const NewMemory = reinterpret_cast<std::byte*>(AutoRTFM::Allocate(TotalSize, alignof(FNode)));
+			auto* const NewMemory = reinterpret_cast<std::byte*>(Allocator::Allocate(TotalSize, alignof(FNode)));
 			auto* const NewInterval = std::launder(reinterpret_cast<FInterval*>(NewMemory + IntervalOffset));
 			auto* const NewLeft = std::launder(reinterpret_cast<NodeIndex*>(NewMemory + LeftOffset));
 			auto* const NewRight = std::launder(reinterpret_cast<NodeIndex*>(NewMemory + RightOffset));
@@ -398,7 +397,7 @@ private:
 				}
 			}
 
-			AutoRTFM::Free(Memory, MemorySizeInBytes);
+			Allocator::Free(Memory, MemorySizeInBytes);
 
 			Ptrs.Interval = NewInterval;
 			Ptrs.Left = NewLeft;
diff --git a/Engine/Source/Runtime/AutoRTFM/Private/Pool.h b/Engine/Source/Runtime/AutoRTFM/Private/Pool.h
--- a/Engine/Source/Runtime/AutoRTFM/Private/Pool.h
+++ b/Engine/Source/Runtime/AutoRTFM/Private/Pool.h
@@ -23,7 +23,7 @@ namespace AutoRTFM
 //   InlineCapacity - the number of items that can be held before spilling to
 //                    the heap.
 template <typename ItemType, size_t InlineCapacity>
-class TPool
+class AUTORTFM_DISABLE_NO_SANITIZE TPool
 {
 public:
 	using FItem = ItemType;
diff --git a/Engine/Source/Runtime/AutoRTFM/Private/Sanitizer.cpp b/Engine/Source/Runtime/AutoRTFM/Private/Sanitizer.cpp
new file mode 100644
--- /dev/null
+++ b/Engine/Source/Runtime/AutoRTFM/Private/Sanitizer.cpp
@@ -0,0 +1,475 @@
+// Copyright Epic Games, Inc. All Rights Reserved.
+
+#if defined(__AUTORTFM) && __AUTORTFM
+
+#include "Sanitizer.h"
+
+#if AUTORTFM_SANITIZER
+
+#include "AutoRTFM.h"
+#include "AutoRTFM/CAPI.h"
+#include "CallstackTree.h"
+#include "ExternAPI.h"
+#include "IntervalTree.h"
+#include "StackRange.h"
+#include "Utils.h"
+
+#include <mutex>
+
+namespace AutoRTFM::Sanitizer
+{
+
+struct AUTORTFM_DISABLE_NO_SANITIZE FOSAllocator
+{
+	static void* Allocate(size_t Size, size_t Alignment)
+	{
+		void* Ptr = malloc(Size);
+		AUTORTFM_ASSERT(Ptr);
+		AUTORTFM_ASSERT((reinterpret_cast<uintptr_t>(Ptr) & (Alignment - 1)) == 0);
+		return Ptr;
+	}
+
+	static void* Reallocate(void* Pointer, size_t Size, size_t Alignment, [[maybe_unused]] size_t PreviousSize)
+	{
+		void* Ptr = realloc(Pointer, Size);
+		AUTORTFM_ASSERT(Ptr);
+		AUTORTFM_ASSERT((reinterpret_cast<uintptr_t>(Ptr) & (Alignment - 1)) == 0);
+		return Ptr;
+	}
+
+	static void* AllocateZeroed(size_t Size, size_t Alignment)
+	{
+		void* Ptr = calloc(1, Size);
+		AUTORTFM_ASSERT(Ptr);
+		AUTORTFM_ASSERT((reinterpret_cast<uintptr_t>(Ptr) & (Alignment - 1)) == 0);
+		return Ptr;
+	}
+
+	static void Free(void* Ptr, [[maybe_unused]] size_t Size)
+	{
+		free(Ptr);
+	}
+};
+static_assert(Allocator<FOSAllocator>);
+
+struct AUTORTFM_DISABLE_NO_SANITIZE FSanitizer
+{
+	struct FLock;
+
+	static std::atomic<Sanitizer::EMode> Mode;
+	static std::atomic<size_t> NumTransactionsRunning;
+	static std::atomic<bool> bRecordClosedWriteCallstacks;
+
+	static void Initialize();
+	static void Shutdown();
+
+	void DisableTransaction();
+	void ReenableTransaction();
+
+	size_t NumberOfIssuesFound() const;
+
+	void StartTransaction(TransactionID ID, FStackRange StackRange);
+	void CommitTransaction(TransactionID ID);
+	void AbortTransaction(TransactionID ID);
+
+	void OpenWrite(void* Ptr, size_t Size, void* ProgramCounter);
+	void ClosedWrite(TransactionID ID, void* Ptr, size_t Size, void* ProgramCounter);
+
+private:
+	using FCallstackTree = TCallstackTree<FOSAllocator>;
+	using TClosedWriteTree = TIntervalTree<FCallstackTree::FHandle, FOSAllocator>;
+
+	struct AUTORTFM_DISABLE_NO_SANITIZE FTransaction
+	{
+		TransactionID const ID;
+		FTransaction* const Parent;
+		FStackRange const StackRange;
+		size_t DisableCount = 0;
+		TClosedWriteTree ClosedWrites;
+	};
+
+	struct FCallstackPair
+	{
+		FCallstackTree::FHandle Open;
+		FCallstackTree::FHandle Closed;
+
+		struct FHasher
+		{
+			size_t operator()(FCallstackPair Pair) const
+			{
+				return (static_cast<size_t>(Pair.Open) << 32) | static_cast<size_t>(Pair.Closed);
+			}
+		};
+		bool operator==(const FCallstackPair&) const = default;
+	};
+
+	static FSanitizer* Instance;
+
+	void PopTransaction(TransactionID ID);
+
+	void ReportConflictingOpenWrite(
+		void* Ptr, size_t Size, void* OpenWriteProgramCounter, FCallstackTree::FHandle ClosedWriteCallstack);
+
+	FCallstackTree::FHandle GetCallstackHandle(void* ProgramCounter);
+
+	void LogCallstack(FCallstackTree::FHandle, autortfm_log_severity);
+
+	FTransaction* CurrentTransaction = nullptr;
+	FCallstackTree CallstackTree;
+	std::unordered_set<FCallstackPair, FCallstackPair::FHasher> ReportedCallstacks;
+	size_t NumIssues = 0;
+};
+
+FSanitizer* FSanitizer::Instance = nullptr;
+std::atomic<Sanitizer::EMode> FSanitizer::Mode{EMode::Error};
+std::atomic<size_t> FSanitizer::NumTransactionsRunning{0};
+std::atomic<bool> FSanitizer::bRecordClosedWriteCallstacks{false};
+
+struct AUTORTFM_DISABLE_NO_SANITIZE FSanitizer::FLock
+{
+	FLock()
+	{
+		if (!bLockedOnThisThread)
+		{
+			Mutex.lock();
+			bLockedOnThisThread = true;
+			Sanitizer = FSanitizer::Instance;
+			bHoldsLock = true;
+		}
+	}
+
+	~FLock()
+	{
+		if (bHoldsLock)
+		{
+			AUTORTFM_ASSERT_DEBUG(bLockedOnThisThread);
+			bLockedOnThisThread = false;
+			Mutex.unlock();
+		}
+	}
+
+	operator bool() const
+	{
+		return Sanitizer;
+	}
+
+	FSanitizer* operator->() const
+	{
+		return Sanitizer;
+	}
+
+private:
+	FLock(const FLock&) = delete;
+	FLock(FLock&&) = delete;
+	const FLock& operator=(const FLock&) = delete;
+	FLock&& operator=(FLock&&) = delete;
+
+	static std::mutex Mutex;
+	static thread_local bool bLockedOnThisThread;
+
+	FSanitizer* Sanitizer = nullptr;
+	bool bHoldsLock = false;
+};
+
+std::mutex FSanitizer::FLock::Mutex;
+thread_local bool FSanitizer::FLock::bLockedOnThisThread = false;
+
+void FSanitizer::Initialize()
+{
+	Instance = New<FOSAllocator, FSanitizer>();
+}
+
+void FSanitizer::Shutdown()
+{
+	Delete<FOSAllocator>(Instance);
+	Instance = nullptr;
+}
+
+void FSanitizer::DisableTransaction()
+{
+	if (FTransaction* Transaction = CurrentTransaction)
+	{
+		Transaction->DisableCount++;
+	}
+}
+
+void FSanitizer::ReenableTransaction()
+{
+	if (FTransaction* Transaction = CurrentTransaction)
+	{
+		AUTORTFM_ASSERT(Transaction->DisableCount > 0);
+		Transaction->DisableCount--;
+	}
+}
+
+size_t FSanitizer::NumberOfIssuesFound() const
+{
+	return NumIssues;
+}
+
+void FSanitizer::StartTransaction(TransactionID ID, FStackRange StackRange)
+{
+	size_t DisableCount = CurrentTransaction ? CurrentTransaction->DisableCount : 0;
+	CurrentTransaction = New<FOSAllocator, FTransaction>(ID, CurrentTransaction, StackRange, DisableCount);
+	NumTransactionsRunning++;
+}
+
+void FSanitizer::CommitTransaction(TransactionID ID)
+{
+	FTransaction* const Transaction = CurrentTransaction;
+
+	if (FTransaction* const Parent = Transaction->Parent)
+	{
+		for (TClosedWriteTree::FIntervalAndData Entry : Transaction->ClosedWrites)
+		{
+			const void* Address = reinterpret_cast<void*>(Entry.Interval.Start);
+			if (!Parent->StackRange.Contains(Address))
+			{
+				Parent->ClosedWrites.Insert(Address, Entry.Interval.Size(), Entry.Data);
+			}
+		}
+	}
+
+	PopTransaction(ID);
+}
+
+void FSanitizer::AbortTransaction(TransactionID ID)
+{
+	PopTransaction(ID);
+}
+
+void FSanitizer::OpenWrite(void* Ptr, size_t Size, void* ProgramCounter)
+{
+	if (!CurrentTransaction || CurrentTransaction->DisableCount)
+	{
+		return;
+	}
+
+	for (FTransaction* Transaction = CurrentTransaction; Transaction; Transaction = Transaction->Parent)
+	{
+		// TODO(SOL-8992): FIntervalTree::Contains() may incorrectly return
+		// false if the queried interval starts outside of an interval added by
+		// FIntervalTree::Insert() but still overlaps an inserted interval.
+		// Fix this.
+		if (const FCallstackTree::FHandle* Callstack = Transaction->ClosedWrites.Get(Ptr, Size); AUTORTFM_UNLIKELY(Callstack))
+		{
+			ReportConflictingOpenWrite(Ptr, Size, ProgramCounter, *Callstack);
+			return;
+		}
+	}
+}
+
+void FSanitizer::ClosedWrite(TransactionID ID, void* Ptr, size_t Size, void* ProgramCounter)
+{
+	FTransaction* const Transaction = CurrentTransaction;
+	AUTORTFM_ASSERT(Transaction->ID == ID);
+
+	FCallstackTree::FHandle const CallstackHandle =
+		bRecordClosedWriteCallstacks ? GetCallstackHandle(ProgramCounter) : FCallstackTree::InvalidHandle;
+	Transaction->ClosedWrites.Insert(Ptr, Size, CallstackHandle);
+}
+
+void FSanitizer::PopTransaction(TransactionID ID)
+{
+	AUTORTFM_ASSERT(ID == CurrentTransaction->ID);
+	FTransaction* const Parent = CurrentTransaction->Parent;
+	Delete<FOSAllocator>(CurrentTransaction);
+	CurrentTransaction = Parent;
+	--NumTransactionsRunning;
+}
+
+void FSanitizer::ReportConflictingOpenWrite(
+	void* Ptr, size_t Size, void* OpenWriteProgramCounter, FCallstackTree::FHandle ClosedWriteCallstack)
+{
+	const autortfm_log_severity Severity = (Mode == Sanitizer::EMode::Error) ? autortfm_log_error : autortfm_log_warn;
+
+	FCallstackTree::FHandle const OpenWriteCallstack = GetCallstackHandle(OpenWriteProgramCounter);
+
+	if (Severity == autortfm_log_warn)
+	{
+		if (!ReportedCallstacks.emplace(OpenWriteCallstack, ClosedWriteCallstack).second)
+		{
+			return;  // Already reported
+		}
+	}
+
+	NumIssues++;
+
+	AUTORTFM_LOG_AT_SEVERITY(Severity, "");
+	AUTORTFM_LOG_AT_SEVERITY(
+		Severity, "********************************************************************************************");
+	AUTORTFM_LOG_AT_SEVERITY(Severity, "AutoRTFMSan: 0x%zu bytes at %p was modified in the closed and now the open.", Size, Ptr);
+	AUTORTFM_LOG_AT_SEVERITY(Severity, "This may lead to memory corruption if the transaction is aborted.");
+	AUTORTFM_LOG_AT_SEVERITY(Severity, "");
+
+	if (ClosedWriteCallstack != FCallstackTree::InvalidHandle)
+	{
+		AUTORTFM_LOG_AT_SEVERITY(Severity, "AutoRTFMSan: The last transactional write to this address was at:");
+		LogCallstack(ClosedWriteCallstack, Severity);
+		AUTORTFM_LOG_AT_SEVERITY(Severity, "");
+	}
+	else if (!bRecordClosedWriteCallstacks)
+	{
+		AUTORTFM_LOG_AT_SEVERITY(Severity, "AutoRTFMSan: Transactional write callstack recording is not enabled.");
+		AUTORTFM_LOG_AT_SEVERITY(Severity, "AutoRTFMSan: Enable to see callstack of the last closed write.");
+	}
+
+	AUTORTFM_LOG_AT_SEVERITY(Severity, "AutoRTFMSan: The open write was at:");
+	LogCallstack(OpenWriteCallstack, Severity);
+	AUTORTFM_LOG_AT_SEVERITY(
+		Severity, "********************************************************************************************");
+	AUTORTFM_LOG_AT_SEVERITY(Severity, "");
+
+	if (Mode == Sanitizer::EMode::Error)
+	{
+		AUTORTFM_FATAL("<AutoRTFMSan fatal error>");
+	}
+}
+
+FSanitizer::FCallstackTree::FHandle FSanitizer::GetCallstackHandle(void* ProgramCounter)
+{
+	static constexpr size_t MaxCallstack = 128;
+	autortfm_callstack_frame Callstack[MaxCallstack];
+	size_t const NumCaptured = GExternAPI.CaptureCallstack(MaxCallstack, Callstack);
+
+	// Attempt to trim off callstack frames under this function.
+	size_t FramesToSkip = 0;
+	for (size_t I = 0, N = std::min(NumCaptured, static_cast<size_t>(8)); I < N; I++)
+	{
+		if (reinterpret_cast<autortfm_callstack_frame>(ProgramCounter) == Callstack[I])
+		{
+			FramesToSkip += I;
+			break;
+		}
+	}
+
+	if (NumCaptured <= FramesToSkip)
+	{
+		return FCallstackTree::InvalidHandle;
+	}
+
+	return CallstackTree.Add(NumCaptured - FramesToSkip, Callstack + FramesToSkip);
+}
+
+void FSanitizer::LogCallstack(FCallstackTree::FHandle Handle, autortfm_log_severity Severity)
+{
+	if (auto Callstack = CallstackTree.Get(Handle); !Callstack.IsEmpty())
+	{
+		GExternAPI.LogCallstack(Severity, Callstack.Num(), &Callstack.Front());
+	}
+}
+
+////////////////////////////////////////////////////////////////////////////////
+// Private interface
+////////////////////////////////////////////////////////////////////////////////
+void Initialize()
+{
+	FSanitizer::Initialize();
+}
+
+void Shutdown()
+{
+	FSanitizer::Shutdown();
+}
+
+size_t NumberOfIssuesFound()
+{
+	FSanitizer::FLock Sanitizer;
+	return Sanitizer->NumberOfIssuesFound();
+}
+
+void StartTransaction(TransactionID ID, FStackRange StackRange)
+{
+	FSanitizer::FLock Sanitizer;
+	Sanitizer->StartTransaction(ID, StackRange);
+}
+
+void CommitTransaction(TransactionID ID)
+{
+	FSanitizer::FLock Sanitizer;
+	Sanitizer->CommitTransaction(ID);
+}
+
+void AbortTransaction(TransactionID ID)
+{
+	FSanitizer::FLock Sanitizer;
+	Sanitizer->AbortTransaction(ID);
+}
+
+void ClosedWrite(TransactionID ID, void* Ptr, size_t Size, void* ProgramCounter)
+{
+	FSanitizer::FLock Sanitizer;
+	Sanitizer->ClosedWrite(ID, Ptr, Size, ProgramCounter);
+}
+
+////////////////////////////////////////////////////////////////////////////////
+// Compiler Interface
+////////////////////////////////////////////////////////////////////////////////
+extern "C" AUTORTFM_DISABLE_NO_SANITIZE UE_AUTORTFM_FORCENOINLINE UE_AUTORTFM_API void autortfm_sanitizer_open_write(
+	void* Ptr, size_t Size) noexcept
+{
+	// Skip the cost of a mutex lock if there are no transactions running, or
+	// if the sanitizer is disabled.
+	if (AnyTransactionsInFlight() && FSanitizer::Mode != EMode::Disabled)
+	{
+		FSanitizer::FLock Sanitizer;
+		if (Sanitizer)
+		{
+			Sanitizer->OpenWrite(Ptr, Size, __builtin_return_address(0));
+		}
+	}
+}
+
+////////////////////////////////////////////////////////////////////////////////
+// Public API
+////////////////////////////////////////////////////////////////////////////////
+EMode GetMode()
+{
+	return FSanitizer::Mode;
+}
+
+void SetMode(EMode NewMode)
+{
+	if (FSanitizer::Mode != NewMode)
+	{
+		FSanitizer::Mode = NewMode;
+		GExternAPI.OnSanitizerModeChanged();
+	}
+}
+
+bool RecordClosedWriteCallstacks()
+{
+	return FSanitizer::bRecordClosedWriteCallstacks;
+}
+
+void SetRecordClosedWriteCallstacks(bool bEnable)
+{
+	if (FSanitizer::bRecordClosedWriteCallstacks != bEnable)
+	{
+		FSanitizer::bRecordClosedWriteCallstacks = bEnable;
+		GExternAPI.OnSanitizerRecordClosedWriteCallstacksChanged();
+	}
+}
+
+void BeginDisable()
+{
+	FSanitizer::FLock Sanitizer;
+	Sanitizer->DisableTransaction();
+}
+
+void EndDisable()
+{
+	FSanitizer::FLock Sanitizer;
+	Sanitizer->ReenableTransaction();
+}
+
+bool AnyTransactionsInFlight()
+{
+	return FSanitizer::NumTransactionsRunning > 0;
+}
+
+}  // namespace AutoRTFM::Sanitizer
+
+#endif  // AUTORTFM_SANITIZER
+
+#endif  // defined(__AUTORTFM) && __AUTORTFM
diff --git a/Engine/Source/Runtime/AutoRTFM/Private/Sanitizer.h b/Engine/Source/Runtime/AutoRTFM/Private/Sanitizer.h
new file mode 100644
--- /dev/null
+++ b/Engine/Source/Runtime/AutoRTFM/Private/Sanitizer.h
@@ -0,0 +1,28 @@
+// Copyright Epic Games, Inc. All Rights Reserved.
+
+#if defined(__AUTORTFM) && __AUTORTFM
+
+#include "AutoRTFM.h"
+
+#if AUTORTFM_SANITIZER
+
+#include "StackRange.h"
+
+namespace AutoRTFM::Sanitizer
+{
+
+AUTORTFM_DISABLE_NO_SANITIZE void Initialize();
+AUTORTFM_DISABLE_NO_SANITIZE void Shutdown();
+
+AUTORTFM_DISABLE_NO_SANITIZE size_t NumberOfIssuesFound();
+
+AUTORTFM_DISABLE_NO_SANITIZE void StartTransaction(TransactionID ID, FStackRange StackRange);
+AUTORTFM_DISABLE_NO_SANITIZE void CommitTransaction(TransactionID ID);
+AUTORTFM_DISABLE_NO_SANITIZE void AbortTransaction(TransactionID ID);
+
+AUTORTFM_DISABLE_NO_SANITIZE void ClosedWrite(TransactionID ID, void* Ptr, size_t Size, void* ProgramCounter);
+
+}  // namespace AutoRTFM::Sanitizer
+
+#endif  //  AUTORTFM_SANITIZER
+#endif  // defined(__AUTORTFM) && __AUTORTFM
diff --git a/Engine/Source/Runtime/AutoRTFM/Private/Stack.h b/Engine/Source/Runtime/AutoRTFM/Private/Stack.h
--- a/Engine/Source/Runtime/AutoRTFM/Private/Stack.h
+++ b/Engine/Source/Runtime/AutoRTFM/Private/Stack.h
@@ -4,8 +4,8 @@
 
 #if (defined(__AUTORTFM) && __AUTORTFM)
 
+#include "Allocator.h"
 #include "ContainerValidation.h"
-#include "ExternAPI.h"
 #include "Utils.h"
 
 #include <cstddef>
@@ -30,7 +30,8 @@ namespace AutoRTFM
 //   InlineCapacity - the number of elements that can be held before spilling to
 //                    the heap.
 //   Validation - if Disabled, do not perform validity assertions.
-template <typename T, size_t InlineCapacity, EContainerValidation Validation = EContainerValidation::Enabled>
+template <typename T, size_t InlineCapacity, Allocator Allocator = FAllocator,
+	EContainerValidation Validation = EContainerValidation::Enabled>
 class AUTORTFM_DISABLE_NO_SANITIZE TStack
 {
 public:
@@ -175,19 +176,19 @@ public:
 		{
 			if (Data == InlineData())
 			{
-				T* NewData = reinterpret_cast<T*>(AutoRTFM::Allocate(NewCapacity * sizeof(T), alignof(T)));
+				T* NewData = reinterpret_cast<T*>(Allocator::Allocate(NewCapacity * sizeof(T), alignof(T)));
 				memcpy(NewData, InlineDataBuffer, Count * sizeof(T));
 				Data = NewData;
 			}
 			else
 			{
 				size_t const PreviousSize = Capacity * sizeof(T);
-				Data = reinterpret_cast<T*>(AutoRTFM::Reallocate(Data, NewCapacity * sizeof(T), alignof(T), PreviousSize));
+				Data = reinterpret_cast<T*>(Allocator::Reallocate(Data, NewCapacity * sizeof(T), alignof(T), PreviousSize));
 			}
 		}
 		else
 		{
-			T* NewData = reinterpret_cast<T*>(AutoRTFM::Allocate(NewCapacity * sizeof(T), alignof(T)));
+			T* NewData = reinterpret_cast<T*>(Allocator::Allocate(NewCapacity * sizeof(T), alignof(T)));
 			for (size_t I = 0, N = Count; I < N; I++)
 			{
 				new (NewData + I) T(std::move(Data[I]));
@@ -267,8 +268,8 @@ public:
 	}
 
 	// Equality operator
-	template <typename OtherT, size_t OtherInlineCapacity>
-	bool operator==(const TStack<OtherT, OtherInlineCapacity>& Other) const
+	template <typename OtherT, size_t OtherInlineCapacity, typename OtherAllocator>
+	bool operator==(const TStack<OtherT, OtherInlineCapacity, OtherAllocator>& Other) const
 	{
 		if (Num() != Other.Num())
 		{
@@ -291,7 +292,7 @@ private:
 	{
 		AUTORTFM_ASSERT(Data != InlineData());
 		size_t const Size = Capacity * sizeof(T);
-		AutoRTFM::Free(Data, Size);
+		Allocator::Free(Data, Size);
 	}
 
 	// Copies the data from Other to this stack.
diff --git a/Engine/Source/Runtime/AutoRTFM/Private/StlAllocator.h b/Engine/Source/Runtime/AutoRTFM/Private/StlAllocator.h
--- a/Engine/Source/Runtime/AutoRTFM/Private/StlAllocator.h
+++ b/Engine/Source/Runtime/AutoRTFM/Private/StlAllocator.h
@@ -2,49 +2,49 @@
 
 #pragma once
 
+#include "AutoRTFM/Constants.h"
 #if (defined(__AUTORTFM) && __AUTORTFM)
 
-#include "ExternAPI.h"
+#include "Allocator.h"
 
 #include <cstddef>
 
 namespace AutoRTFM
 {
 
-// A stl-compatible allocator which allocates with AutoRTFM::Allocate() and
-// frees with AutoRTFM::Free().
+// A stl-compatible allocator which wraps an AutoRTFM allocator.
 // TODO(SOL-7652): Remove this once HashMap has been replaced with a bespoke
 // hashmap implementation.
-template <typename T>
-class StlAllocator
+template <typename T, Allocator Allocator = FAllocator>
+class [[clang::autortfm(autortfm_mode_internal)]] StlAllocator
 {
 public:
 	using value_type = T;
 
 	StlAllocator() = default;
 
-	template <typename U>
-	StlAllocator(const StlAllocator<U>&){};
+	template <typename U, typename A>
+	StlAllocator(const StlAllocator<U, A>&){};
 
 	T* allocate(size_t Count)
 	{
-		return static_cast<T*>(AutoRTFM::Allocate(Count * sizeof(T), alignof(T)));
+		return static_cast<T*>(Allocator::Allocate(Count * sizeof(T), alignof(T)));
 	}
 
 	void deallocate(T* Pointer, size_t Count)
 	{
-		AutoRTFM::Free(Pointer, Count * sizeof(T));
+		Allocator::Free(Pointer, Count * sizeof(T));
 	}
 };
 
-template <typename T, typename U>
-bool operator==(const StlAllocator<T>&, const StlAllocator<U>&)
+template <typename T, typename U, typename A>
+bool operator==(const StlAllocator<T, A>&, const StlAllocator<U, A>&)
 {
 	return true;
 }
 
-template <typename T, typename U>
-bool operator!=(const StlAllocator<T>&, const StlAllocator<U>&)
+template <typename T, typename U, typename A>
+bool operator!=(const StlAllocator<T, A>&, const StlAllocator<U, A>&)
 {
 	return false;
 }
diff --git a/Engine/Source/Runtime/AutoRTFM/Private/Transaction.cpp b/Engine/Source/Runtime/AutoRTFM/Private/Transaction.cpp
--- a/Engine/Source/Runtime/AutoRTFM/Private/Transaction.cpp
+++ b/Engine/Source/Runtime/AutoRTFM/Private/Transaction.cpp
@@ -10,14 +10,9 @@
 #include "ContextInlines.h"
 #include "HitSet.h"
 #include "MetricsPriv.h"
+#include "Sanitizer.h"
 #include "Utils.h"
 
-#if PLATFORM_WINDOWS
-#include "WindowsHeader.h"
-#endif
-
-#include <cstdarg>
-#include <cstddef>
 #include <cstdint>
 
 namespace AutoRTFM
@@ -57,9 +52,6 @@ void FTransaction::Suppress()
 
 	bIsInAllocateFn = false;
 
-	AUTORTFM_SANITIZER_ONLY(ClosedWrites.Reset());
-	AUTORTFM_SANITIZER_ONLY(SanitizerDisableCount = 0);
-
 	AUTORTFM_ASSERT(IsFresh());
 }
 
@@ -80,6 +72,7 @@ void FTransaction::Initialize(FTransaction* InParent, FStackRange InStackRange)
 {
 	Parent = InParent;
 	StackRange = InStackRange;
+	AUTORTFM_SANITIZER_ONLY(Sanitizer::StartTransaction(ID, StackRange));
 
 	if (Parent)
 	{
@@ -95,11 +88,7 @@ bool FTransaction::IsFresh() const
 {
 	return HitSet.IsEmpty() && NewMemoryTracker.IsEmpty() && WriteLog.IsEmpty() && CommitTasks.IsEmpty() && PreAbortTasks.IsEmpty()
 		&& AbortTasks.IsEmpty() && DeferredPopOnCommitHandlers.IsEmpty() && DeferredPopOnAbortHandlers.IsEmpty()
-		&& DeferredPopAllOnCommitHandlers.IsEmpty() && DeferredPopAllOnAbortHandlers.IsEmpty()
-#if AUTORTFM_SANITIZER
-		&& ClosedWrites.IsEmpty() && SanitizerDisableCount == 0
-#endif
-		;
+		&& DeferredPopAllOnCommitHandlers.IsEmpty() && DeferredPopAllOnAbortHandlers.IsEmpty();
 }
 
 void FTransaction::Abort()
@@ -127,6 +116,8 @@ void FTransaction::Abort()
 	// Revert all memory in the write-log.
 	Undo();
 
+	AUTORTFM_SANITIZER_ONLY(Sanitizer::AbortTransaction(ID));
+
 	// Call each of the post memory rollback on-abort handlers.
 	AbortTasks.RemoveEachBackward([&](TTask<void()>& Task) { Task(); });
 }
@@ -136,6 +127,7 @@ void FTransaction::Commit()
 	AUTORTFM_ASSERT(IsActive());
 	AUTORTFM_ASSERT(Context->GetStatus() == EContextStatus::Committing);
 	AUTORTFM_ASSERT(Context->GetCurrentTransaction() == this);
+	AUTORTFM_SANITIZER_ONLY(Sanitizer::CommitTransaction(ID));
 
 	Stats.Collect<EStatsKind::Commit>();
 	CollectStats();
@@ -157,7 +149,6 @@ void FTransaction::Undo()
 
 	// Prevent custom rollback functions from triggering the AutoRTFM sanitizer
 	AUTORTFM_SANITIZER_DISABLE_SCOPE();
-	AUTORTFM_SANITIZER_ONLY(ClosedWrites.Reset());
 
 	for (auto Iter = WriteLog.rbegin(); Iter != WriteLog.rend(); ++Iter)
 	{
@@ -223,17 +214,6 @@ void FTransaction::CommitNested()
 		}
 	}
 
-#if AUTORTFM_SANITIZER
-	for (TClosedWriteTree::FIntervalAndData Entry : ClosedWrites)
-	{
-		const void* Address = reinterpret_cast<void*>(Entry.Interval.Start);
-		if (!Parent->IsOnStack(Address))
-		{
-			Parent->ClosedWrites.Insert(Address, Entry.Interval.Size(), Entry.Data);
-		}
-	}
-#endif
-
 	// This point is likely the high-water mark for this transactional nest's memory usage.
 	// It's after we've cloned our write log and hit set into the parent transaction,
 	// but before we've cleared our deferred-task list.
@@ -291,77 +271,9 @@ void FTransaction::CommitOuterNest()
 	PreAbortTasks.Reset();
 	AbortTasks.Reset();
 
-	AUTORTFM_SANITIZER_ONLY(ClosedWrites.Reset());
-
 	CommitTasks.RemoveEachForward([](TTask<void()>& Task) { Task(); });
 }
 
-#if AUTORTFM_SANITIZER
-void FTransaction::SanitizerReportConflictingOpenWrite(void* Ptr, size_t Size, FCallstackTree::FHandle ClosedWriteCallstack)
-{
-	Sanitizer::EMode SanitizerMode = this->SanitizerMode();
-
-	const autortfm_log_severity Severity = (SanitizerMode == Sanitizer::EMode::Error) ? autortfm_log_error : autortfm_log_warn;
-
-	FCallstackTree::FHandle const OpenWriteCallstack = Context->GetCallstackHandle(/* FramesToSkip */ 2);
-
-	if (Severity == autortfm_log_warn)
-	{
-		struct FCallstackPair
-		{
-			FCallstackTree::FHandle Open;
-			FCallstackTree::FHandle Closed;
-
-			struct FHasher
-			{
-				size_t operator()(FCallstackPair Pair) const
-				{
-					return (static_cast<size_t>(Pair.Open) << 32) | static_cast<size_t>(Pair.Closed);
-				}
-			};
-			bool operator==(const FCallstackPair&) const = default;
-		};
-		static std::unordered_set<FCallstackPair, FCallstackPair::FHasher> Reported;
-		if (!Reported.emplace(OpenWriteCallstack, ClosedWriteCallstack).second)
-		{
-			return;  // Already reported
-		}
-	}
-
-	Context->IncrementNumSanitizerIssues();
-
-	AUTORTFM_LOG_AT_SEVERITY(Severity, "");
-	AUTORTFM_LOG_AT_SEVERITY(
-		Severity, "********************************************************************************************");
-	AUTORTFM_LOG_AT_SEVERITY(Severity, "AutoRTFMSan: 0x%zu bytes at %p was modified in the closed and now the open.", Size, Ptr);
-	AUTORTFM_LOG_AT_SEVERITY(Severity, "This may lead to memory corruption if the transaction is aborted.");
-	AUTORTFM_LOG_AT_SEVERITY(Severity, "");
-
-	if (ClosedWriteCallstack != FCallstackTree::InvalidHandle)
-	{
-		AUTORTFM_LOG_AT_SEVERITY(Severity, "AutoRTFMSan: The last closed write to this address was at:");
-		Context->LogCallstack(ClosedWriteCallstack, Severity);
-		AUTORTFM_LOG_AT_SEVERITY(Severity, "");
-	}
-	else if (!AutoRTFM::Sanitizer::RecordClosedWriteCallstacks())
-	{
-		AUTORTFM_LOG_AT_SEVERITY(Severity, "AutoRTFMSan: Closed write callstack recording is not enabled.");
-		AUTORTFM_LOG_AT_SEVERITY(Severity, "AutoRTFMSan: Enable to see callstack of the last closed write.");
-	}
-
-	AUTORTFM_LOG_AT_SEVERITY(Severity, "AutoRTFMSan: The open write was at:");
-	Context->LogCallstack(OpenWriteCallstack, Severity);
-	AUTORTFM_LOG_AT_SEVERITY(
-		Severity, "********************************************************************************************");
-	AUTORTFM_LOG_AT_SEVERITY(Severity, "");
-
-	if (SanitizerMode == Sanitizer::EMode::Error)
-	{
-		AUTORTFM_FATAL("<AutoRTFMSan fatal error>");
-	}
-}
-#endif  // AUTORTFM_SANITIZER
-
 }  // namespace AutoRTFM
 
 #endif  // defined(__AUTORTFM) && __AUTORTFM
diff --git a/Engine/Source/Runtime/AutoRTFM/Private/Transaction.h b/Engine/Source/Runtime/AutoRTFM/Private/Transaction.h
--- a/Engine/Source/Runtime/AutoRTFM/Private/Transaction.h
+++ b/Engine/Source/Runtime/AutoRTFM/Private/Transaction.h
@@ -15,12 +15,6 @@
 #include "Utils.h"
 #include "WriteLog.h"
 
-#if AUTORTFM_SANITIZER
-#include "CallstackTree.h"
-#endif  // AUTORTFM_SANITIZER
-
-#include <optional>
-
 namespace AutoRTFM
 {
 class FContext;
@@ -93,34 +87,6 @@ public:
 	// Returns true if the LogicalAddress is within the stack of the transaction.
 	inline bool IsOnStack(const void* LogicalAddress) const;
 
-#if AUTORTFM_SANITIZER
-	inline void SanitizerOpenWrite(void* Ptr, size_t Size);
-
-	inline void DisableSanitizer()
-	{
-		SanitizerDisableCount++;
-	}
-
-	inline void EnableSanitizer()
-	{
-		AUTORTFM_ASSERT(SanitizerDisableCount > 0);
-		SanitizerDisableCount--;
-	}
-
-	inline Sanitizer::EMode SanitizerMode() const
-	{
-		if (SanitizerDisableCount)
-		{
-			return Sanitizer::EMode::Disabled;
-		}
-		if (Parent)
-		{
-			return Parent->SanitizerMode();
-		}
-		return Sanitizer::GetMode();
-	}
-#endif
-
 private:
 	using FTaskArray = TTaskArray<TTask<void()>>;
 
@@ -133,10 +99,6 @@ private:
 
 	bool ShouldRecordWrite(void* LogicalAddress) const;
 
-#if AUTORTFM_SANITIZER
-	void SanitizerReportConflictingOpenWrite(void* Ptr, size_t Size, FCallstackTree::FHandle ClosedWriteCallstack);
-#endif
-
 	FContext* Context = nullptr;  // null indicates the transaction is not active
 
 	// This field is used for two purposes:
@@ -174,11 +136,6 @@ private:
 
 	FHitSet HitSet;
 	FIntervalTree NewMemoryTracker;
-#if AUTORTFM_SANITIZER
-	using TClosedWriteTree = TIntervalTree<FCallstackTree::FHandle>;
-	uint32_t SanitizerDisableCount = 0;
-	TClosedWriteTree ClosedWrites;
-#endif
 	FWriteLog WriteLog;
 	TStatStorage<uint64_t> StatDepth = 1;
 	FStackRange StackRange;
diff --git a/Engine/Source/Runtime/AutoRTFM/Private/TransactionInlines.h b/Engine/Source/Runtime/AutoRTFM/Private/TransactionInlines.h
--- a/Engine/Source/Runtime/AutoRTFM/Private/TransactionInlines.h
+++ b/Engine/Source/Runtime/AutoRTFM/Private/TransactionInlines.h
@@ -6,8 +6,8 @@
 
 #include "AutoRTFM/Defines.h"
 #include "AutoRTFM/ReportHazard.h"
-#include "Context.h"
 #include "HitSet.h"
+#include "Sanitizer.h"
 #include "ScopedGuard.h"
 #include "Transaction.h"
 #include "Utils.h"
@@ -126,9 +126,7 @@ AUTORTFM_NO_ASAN UE_AUTORTFM_FORCEINLINE void FTransaction::RecordWrite(
 #if AUTORTFM_SANITIZER
 	if ((Flags & EWriteFlags::NoSanitize) != EWriteFlags::NoSanitize)
 	{
-		FCallstackTree::FHandle CallstackHandle =
-			AutoRTFM::Sanitizer::RecordClosedWriteCallstacks() ? Context->GetCallstackHandle() : FCallstackTree::InvalidHandle;
-		ClosedWrites.Insert(LogicalAddress, Size, CallstackHandle);
+		Sanitizer::ClosedWrite(ID, LogicalAddress, Size, __builtin_return_address(0));
 	}
 #endif
 }
@@ -192,11 +190,7 @@ AUTORTFM_NO_ASAN UE_AUTORTFM_FORCENOINLINE void FTransaction::RecordWriteInserte
 
 	WriteLog.PushSmall<SIZE>(static_cast<std::byte*>(LogicalAddress), static_cast<std::byte*>(LogicalAddress));
 
-#if AUTORTFM_SANITIZER
-	FCallstackTree::FHandle CallstackHandle =
-		AutoRTFM::Sanitizer::RecordClosedWriteCallstacks() ? Context->GetCallstackHandle() : FCallstackTree::InvalidHandle;
-	ClosedWrites.Insert(LogicalAddress, SIZE, CallstackHandle);
-#endif
+	AUTORTFM_SANITIZER_ONLY(Sanitizer::ClosedWrite(ID, LogicalAddress, SIZE, __builtin_return_address(0)));
 }
 
 UE_AUTORTFM_FORCEINLINE void FTransaction::DidAllocate(void* LogicalAddress, const size_t Size)
@@ -330,35 +324,6 @@ UE_AUTORTFM_FORCEINLINE void FTransaction::CollectStats() const
 	Stats.Collect<EStatsKind::AverageHitSetCapacity>(HitSet.GetCapacity());
 }
 
-#if AUTORTFM_SANITIZER
-void FTransaction::SanitizerOpenWrite(void* Ptr, size_t Size)
-{
-	if (bIsInAllocateFn)
-	{
-		// Open write was performed while we were allocating internal AutoRTFM
-		// memory. Ignore.
-		return;
-	}
-
-	if (this->SanitizerMode() == Sanitizer::EMode::Disabled)
-	{
-		return;
-	}
-
-	for (FTransaction* Transaction = this; Transaction; Transaction = Transaction->Parent)
-	{
-		// TODO(SOL-8801): FIntervalTree::Contains() may incorrectly return
-		// false if the queried interval starts outside of an interval added by
-		// FIntervalTree::Insert() but still overlaps an inserted interval.
-		// Fix this.
-		if (const FCallstackTree::FHandle* Callstack = Transaction->ClosedWrites.Get(Ptr, Size); AUTORTFM_UNLIKELY(Callstack))
-		{
-			SanitizerReportConflictingOpenWrite(Ptr, Size, *Callstack);
-		}
-	}
-}
-#endif  // AUTORTFM_SANITIZER
-
 }  // namespace AutoRTFM
 
 #undef AUTORTFM_CHECK_ASAN_FAKE_STACKS
diff --git a/Engine/Source/Runtime/AutoRTFM/Private/WriteLog.h b/Engine/Source/Runtime/AutoRTFM/Private/WriteLog.h
--- a/Engine/Source/Runtime/AutoRTFM/Private/WriteLog.h
+++ b/Engine/Source/Runtime/AutoRTFM/Private/WriteLog.h
@@ -4,6 +4,7 @@
 
 #if (defined(__AUTORTFM) && __AUTORTFM)
 
+#include "Allocator.h"
 #include "AutoRTFM.h"
 #include "BuildMacros.h"
 #include "Utils.h"
@@ -221,7 +222,7 @@ private:
 		static FBlock* Allocate(size_t TotalSize)
 		{
 			AUTORTFM_ASSERT(TotalSize > (sizeof(FBlock) + sizeof(FRecord)));
-			void* Memory = AutoRTFM::Allocate(TotalSize, alignof(FBlock));
+			void* Memory = FAllocator::Allocate(TotalSize, alignof(FBlock));
 			// Disable false-positive warning C6386: Buffer overrun while writing to 'Memory'
 			CA_SUPPRESS(6386)
 			return new (Memory) FBlock(TotalSize);
@@ -233,7 +234,7 @@ private:
 		{
 			std::byte* const End = reinterpret_cast<std::byte*>(FirstRecord + 1);
 			size_t AllocationSize = static_cast<size_t>(End - reinterpret_cast<std::byte*>(this));
-			AutoRTFM::Free(this, AllocationSize);
+			FAllocator::Free(this, AllocationSize);
 		}
 
 		// Returns a pointer to the data for the first entry
diff --git a/Engine/Source/Runtime/AutoRTFM/Public/AutoRTFM.h b/Engine/Source/Runtime/AutoRTFM/Public/AutoRTFM.h
--- a/Engine/Source/Runtime/AutoRTFM/Public/AutoRTFM.h
+++ b/Engine/Source/Runtime/AutoRTFM/Public/AutoRTFM.h
@@ -263,10 +263,16 @@ enum class EMode : uint8_t
 };
 
 // Returns the global AutoRTFM sanitizer mode currently in use.
-UE_AUTORTFM_API EMode GetMode();
+AUTORTFM_OPEN UE_AUTORTFM_API EMode GetMode();
 
 // Sets the global AutoRTFM sanitizer mode.
-UE_AUTORTFM_API void SetMode(EMode Mode);
+AUTORTFM_OPEN UE_AUTORTFM_API void SetMode(EMode Mode);
+
+// Returns true if the AutoRTFM sanitizer records the callstack of closed writes.
+AUTORTFM_OPEN UE_AUTORTFM_API bool RecordClosedWriteCallstacks();
+
+// Sets whether the AutoRTFM sanitizer records the callstack of closed writes.
+AUTORTFM_OPEN UE_AUTORTFM_API void SetRecordClosedWriteCallstacks(bool bEnable);
 
 // Temporarily disables the AutoRTFM sanitizer for the current transaction
 // and any child transactions.
@@ -280,11 +286,10 @@ AUTORTFM_OPEN UE_AUTORTFM_API void BeginDisable();
 // See BeginDisable().
 AUTORTFM_OPEN UE_AUTORTFM_API void EndDisable();
 
-// Returns true if the AutoRTFM sanitizer records the callstack of closed writes.
-UE_AUTORTFM_API bool RecordClosedWriteCallstacks();
-
-// Sets whether the AutoRTFM sanitizer records the callstack of closed writes.
-UE_AUTORTFM_API void SetRecordClosedWriteCallstacks(bool bEnable);
+// Returns true if there are any AutoRTFM transactions in flight.
+// Can be safely called from any thread, although the result will be
+// non-deterministic without synchronization between threads.
+AUTORTFM_OPEN UE_AUTORTFM_API bool AnyTransactionsInFlight();
 
 // A helper class that temporarily disables the AutoRTFM sanitizer for
 // the lifetime of the object.
