[AutoRTFM] Fix allocation alignment of BlockAllocator

Spotted by BugHawk: Incorrect alignment logic in bump allocation: TryAllocate aligns the Remaining byte-count (AlignDown(Remaining, Alignment)) instead of aligning the returned pointer address. This can return misaligned pointers whenever the block end address is not already aligned to Alignment.

#Synced-CL 50350561

diff --git a/Engine/Source/Runtime/AutoRTFM/Private/BlockAllocator.h b/Engine/Source/Runtime/AutoRTFM/Private/BlockAllocator.h
--- a/Engine/Source/Runtime/AutoRTFM/Private/BlockAllocator.h
+++ b/Engine/Source/Runtime/AutoRTFM/Private/BlockAllocator.h
@@ -43,7 +43,7 @@ public:
 			FAllocator::Free(Tail, Tail->TotalSize());
 			Tail = Prev;
 		}
-		InlineBlock.Remaining = InlineBlockDataSize;
+		InlineBlock.NextAllocation = reinterpret_cast<std::byte*>(&InlineBlock) + sizeof(FBlockHeader);
 		NextBlockSize = InlineBlockDataSize * GrowthPercentage / 100;
 	}
 
@@ -81,9 +81,9 @@ private:
 		// The previous FBlock in the singly linked-list
 		FBlockHeader* const Prev = nullptr;
 		// The size of the block's data
-		const size_t BlockDataSize = 0;
-		// The number of unallocated bytes remaining in the block
-		size_t Remaining = BlockDataSize;
+		size_t const BlockDataSize;
+		// Pointer to the next sub-allocation.
+		std::byte* NextAllocation = reinterpret_cast<std::byte*>(this) + sizeof(FBlockHeader);
 		// <data>
 
 		// Allocates, constructs and returns a new block from the heap-allocated memory.
@@ -103,13 +103,15 @@ private:
 		// the sub-allocated memory on success, or nullptr on failure.
 		inline void* TryAllocate(size_t Size, size_t Alignment)
 		{
-			size_t RemainingAligned = AlignDown(Remaining, Alignment);
-			if (RemainingAligned < Size)
+			std::byte* const BlockDataStart = reinterpret_cast<std::byte*>(this) + sizeof(FBlockHeader);
+			std::byte* const BlockDataEnd = BlockDataStart + BlockDataSize;
+			std::byte* const Ptr = reinterpret_cast<std::byte*>(AlignUp(reinterpret_cast<uintptr_t>(NextAllocation), Alignment));
+			std::byte* const End = Ptr + Size;
+			if (End > BlockDataEnd || /* overflow */ AUTORTFM_UNLIKELY(End < BlockDataStart))
 			{
 				return nullptr;
 			}
-			void* Ptr = reinterpret_cast<std::byte*>(this) + sizeof(FBlockHeader) + BlockDataSize - RemainingAligned;
-			Remaining = RemainingAligned - Size;
+			NextAllocation = End;
 			return Ptr;
 		}
 	};
