[AutoRTFM] Fix deadlock when the sanitizer finds a conflicting write, and the sanitizer mode is to error (fatal).

Move the logging outside of the FSanitizer mutex lock, otherwise logging the error can attempt to re-lock the mutex, and we deadlock.

#Synced-CL 50428254

diff --git a/Engine/Source/Runtime/AutoRTFM/Private/CallstackTree.h b/Engine/Source/Runtime/AutoRTFM/Private/CallstackTree.h
--- a/Engine/Source/Runtime/AutoRTFM/Private/CallstackTree.h
+++ b/Engine/Source/Runtime/AutoRTFM/Private/CallstackTree.h
@@ -24,6 +24,7 @@ class AUTORTFM_DISABLE_NO_SANITIZE TCallstackTree final
 {
 public:
 	using FHandle = uint32_t;
+	using FCallstack = TStack<autortfm_callstack_frame, 8, Allocator>;
 	static constexpr FHandle InvalidHandle = 0;
 	static constexpr FHandle MaxHandle = ~InvalidHandle;
 
@@ -56,9 +57,9 @@ public:
 	}
 
 	// Retrieves a callstack previously added with Add().
-	TStack<autortfm_callstack_frame, 8, Allocator> Get(FHandle Handle) const
+	FCallstack Get(FHandle Handle) const
 	{
-		TStack<autortfm_callstack_frame, 8, Allocator> StackTrace;
+		FCallstack StackTrace;
 		while (Handle != InvalidHandle)
 		{
 			FEntry const& Entry = Entries[Handle - 1];
diff --git a/Engine/Source/Runtime/AutoRTFM/Private/Sanitizer.cpp b/Engine/Source/Runtime/AutoRTFM/Private/Sanitizer.cpp
--- a/Engine/Source/Runtime/AutoRTFM/Private/Sanitizer.cpp
+++ b/Engine/Source/Runtime/AutoRTFM/Private/Sanitizer.cpp
@@ -14,6 +14,7 @@
 #include "StackRange.h"
 #include "Utils.h"
 
+#include <memory>
 #include <mutex>
 
 namespace AutoRTFM::Sanitizer
@@ -49,11 +50,39 @@ struct AUTORTFM_DISABLE_NO_SANITIZE FOSAllocator
 	{
 		free(Ptr);
 	}
+
+	// Deleter that can be used as the default destruction policy of std::unique_ptr.
+	struct FDeleter
+	{
+		template <typename T>
+		void operator()(T* Ptr) const
+		{
+			Delete<FOSAllocator>(Ptr);
+		}
+	};
 };
 static_assert(Allocator<FOSAllocator>);
 
 struct AUTORTFM_DISABLE_NO_SANITIZE FSanitizer
 {
+	using FCallstackTree = TCallstackTree<FOSAllocator>;
+
+	struct FConflictingWrite
+	{
+		using FPtr = std::unique_ptr<FConflictingWrite, FOSAllocator::FDeleter>;
+
+		autortfm_log_severity const Severity;
+		void* const Ptr;
+		size_t const Size;
+		FCallstackTree::FCallstack const OpenWriteCallstack;
+		FCallstackTree::FCallstack const ClosedWriteCallstack;
+
+		// Report the conflicting write, aborting the application if severity is autortfm_log_fatal.
+		// As logging the conflicting write may call back into autortfm_sanitizer_open_write(), this
+		// must be called without the FSanitizer lock held.
+		void Report() const;
+	};
+
 	struct FLock;
 
 	static std::atomic<Sanitizer::EMode> Mode;
@@ -72,11 +101,10 @@ struct AUTORTFM_DISABLE_NO_SANITIZE FSanitizer
 	void CommitTransaction(TransactionID ID);
 	void AbortTransaction(TransactionID ID);
 
-	void OpenWrite(void* Ptr, size_t Size, void* ProgramCounter);
+	[[nodiscard]] FConflictingWrite::FPtr OpenWrite(void* Ptr, size_t Size, void* ProgramCounter);
 	void ClosedWrite(TransactionID ID, void* Ptr, size_t Size, void* ProgramCounter);
 
 private:
-	using FCallstackTree = TCallstackTree<FOSAllocator>;
 	using TClosedWriteTree = TIntervalTree<FCallstackTree::FHandle, FOSAllocator>;
 
 	struct AUTORTFM_DISABLE_NO_SANITIZE FTransaction
@@ -107,13 +135,8 @@ private:
 
 	void PopTransaction(TransactionID ID);
 
-	void ReportConflictingOpenWrite(
-		void* Ptr, size_t Size, void* OpenWriteProgramCounter, FCallstackTree::FHandle ClosedWriteCallstack);
-
 	FCallstackTree::FHandle GetCallstackHandle(void* ProgramCounter);
 
-	void LogCallstack(FCallstackTree::FHandle, autortfm_log_severity);
-
 	FTransaction* CurrentTransaction = nullptr;
 	FCallstackTree CallstackTree;
 	std::unordered_set<FCallstackPair, FCallstackPair::FHasher> ReportedCallstacks;
@@ -238,11 +261,11 @@ void FSanitizer::AbortTransaction(TransactionID ID)
 	PopTransaction(ID);
 }
 
-void FSanitizer::OpenWrite(void* Ptr, size_t Size, void* ProgramCounter)
+[[nodiscard]] FSanitizer::FConflictingWrite::FPtr FSanitizer::OpenWrite(void* Ptr, size_t Size, void* ProgramCounter)
 {
 	if (!CurrentTransaction || CurrentTransaction->DisableCount)
 	{
-		return;
+		return nullptr;
 	}
 
 	for (FTransaction* Transaction = CurrentTransaction; Transaction; Transaction = Transaction->Parent)
@@ -251,12 +274,28 @@ void FSanitizer::OpenWrite(void* Ptr, size_t Size, void* ProgramCounter)
 		// false if the queried interval starts outside of an interval added by
 		// FIntervalTree::Insert() but still overlaps an inserted interval.
 		// Fix this.
-		if (const FCallstackTree::FHandle* Callstack = Transaction->ClosedWrites.Get(Ptr, Size); AUTORTFM_UNLIKELY(Callstack))
+		if (FCallstackTree::FHandle const* const ClosedWriteCallstack = Transaction->ClosedWrites.Get(Ptr, Size);
+			AUTORTFM_UNLIKELY(ClosedWriteCallstack))
 		{
-			ReportConflictingOpenWrite(Ptr, Size, ProgramCounter, *Callstack);
-			return;
+			autortfm_log_severity const Severity = (Mode == Sanitizer::EMode::Error) ? autortfm_log_error : autortfm_log_warn;
+
+			FCallstackTree::FHandle const OpenWriteCallstack = GetCallstackHandle(ProgramCounter);
+			if (Severity == autortfm_log_warn)
+			{
+				if (!ReportedCallstacks.emplace(OpenWriteCallstack, *ClosedWriteCallstack).second)
+				{
+					break;  // Already reported
+				}
+			}
+
+			NumIssues++;
+
+			return FConflictingWrite::FPtr(New<FOSAllocator, FConflictingWrite>(
+				Severity, Ptr, Size, CallstackTree.Get(OpenWriteCallstack), CallstackTree.Get(*ClosedWriteCallstack)));
 		}
 	}
+
+	return nullptr;
 }
 
 void FSanitizer::ClosedWrite(TransactionID ID, void* Ptr, size_t Size, void* ProgramCounter)
@@ -278,23 +317,33 @@ void FSanitizer::PopTransaction(TransactionID ID)
 	--NumTransactionsRunning;
 }
 
-void FSanitizer::ReportConflictingOpenWrite(
-	void* Ptr, size_t Size, void* OpenWriteProgramCounter, FCallstackTree::FHandle ClosedWriteCallstack)
+FSanitizer::FCallstackTree::FHandle FSanitizer::GetCallstackHandle(void* ProgramCounter)
 {
-	const autortfm_log_severity Severity = (Mode == Sanitizer::EMode::Error) ? autortfm_log_error : autortfm_log_warn;
-
-	FCallstackTree::FHandle const OpenWriteCallstack = GetCallstackHandle(OpenWriteProgramCounter);
+	static constexpr size_t MaxCallstack = 128;
+	autortfm_callstack_frame Callstack[MaxCallstack];
+	size_t const NumCaptured = GExternAPI.CaptureCallstack(MaxCallstack, Callstack);
 
-	if (Severity == autortfm_log_warn)
+	// Attempt to trim off callstack frames under this function.
+	size_t FramesToSkip = 0;
+	for (size_t I = 0, N = std::min(NumCaptured, static_cast<size_t>(8)); I < N; I++)
 	{
-		if (!ReportedCallstacks.emplace(OpenWriteCallstack, ClosedWriteCallstack).second)
+		if (reinterpret_cast<autortfm_callstack_frame>(ProgramCounter) == Callstack[I])
 		{
-			return;  // Already reported
+			FramesToSkip += I;
+			break;
 		}
 	}
 
-	NumIssues++;
+	if (NumCaptured <= FramesToSkip)
+	{
+		return FCallstackTree::InvalidHandle;
+	}
 
+	return CallstackTree.Add(NumCaptured - FramesToSkip, Callstack + FramesToSkip);
+}
+
+void FSanitizer::FConflictingWrite::Report() const
+{
 	AUTORTFM_LOG_AT_SEVERITY(Severity, "");
 	AUTORTFM_LOG_AT_SEVERITY(
 		Severity, "********************************************************************************************");
@@ -302,10 +351,10 @@ void FSanitizer::ReportConflictingOpenWrite(
 	AUTORTFM_LOG_AT_SEVERITY(Severity, "This may lead to memory corruption if the transaction is aborted.");
 	AUTORTFM_LOG_AT_SEVERITY(Severity, "");
 
-	if (ClosedWriteCallstack != FCallstackTree::InvalidHandle)
+	if (!ClosedWriteCallstack.IsEmpty())
 	{
 		AUTORTFM_LOG_AT_SEVERITY(Severity, "AutoRTFMSan: The last transactional write to this address was at:");
-		LogCallstack(ClosedWriteCallstack, Severity);
+		GExternAPI.LogCallstack(Severity, ClosedWriteCallstack.Num(), &ClosedWriteCallstack.Front());
 		AUTORTFM_LOG_AT_SEVERITY(Severity, "");
 	}
 	else if (!bRecordClosedWriteCallstacks)
@@ -315,7 +364,14 @@ void FSanitizer::ReportConflictingOpenWrite(
 	}
 
 	AUTORTFM_LOG_AT_SEVERITY(Severity, "AutoRTFMSan: The open write was at:");
-	LogCallstack(OpenWriteCallstack, Severity);
+	if (!OpenWriteCallstack.IsEmpty())
+	{
+		GExternAPI.LogCallstack(Severity, OpenWriteCallstack.Num(), &OpenWriteCallstack.Front());
+	}
+	else
+	{
+		AUTORTFM_LOG_AT_SEVERITY(Severity, "<no callstack>");
+	}
 	AUTORTFM_LOG_AT_SEVERITY(
 		Severity, "********************************************************************************************");
 	AUTORTFM_LOG_AT_SEVERITY(Severity, "");
@@ -326,39 +382,6 @@ void FSanitizer::ReportConflictingOpenWrite(
 	}
 }
 
-FSanitizer::FCallstackTree::FHandle FSanitizer::GetCallstackHandle(void* ProgramCounter)
-{
-	static constexpr size_t MaxCallstack = 128;
-	autortfm_callstack_frame Callstack[MaxCallstack];
-	size_t const NumCaptured = GExternAPI.CaptureCallstack(MaxCallstack, Callstack);
-
-	// Attempt to trim off callstack frames under this function.
-	size_t FramesToSkip = 0;
-	for (size_t I = 0, N = std::min(NumCaptured, static_cast<size_t>(8)); I < N; I++)
-	{
-		if (reinterpret_cast<autortfm_callstack_frame>(ProgramCounter) == Callstack[I])
-		{
-			FramesToSkip += I;
-			break;
-		}
-	}
-
-	if (NumCaptured <= FramesToSkip)
-	{
-		return FCallstackTree::InvalidHandle;
-	}
-
-	return CallstackTree.Add(NumCaptured - FramesToSkip, Callstack + FramesToSkip);
-}
-
-void FSanitizer::LogCallstack(FCallstackTree::FHandle Handle, autortfm_log_severity Severity)
-{
-	if (auto Callstack = CallstackTree.Get(Handle); !Callstack.IsEmpty())
-	{
-		GExternAPI.LogCallstack(Severity, Callstack.Num(), &Callstack.Front());
-	}
-}
-
 ////////////////////////////////////////////////////////////////////////////////
 // Private interface
 ////////////////////////////////////////////////////////////////////////////////
@@ -412,10 +435,18 @@ extern "C" AUTORTFM_DISABLE_NO_SANITIZE UE_AUTORTFM_FORCENOINLINE UE_AUTORTFM_AP
 	// if the sanitizer is disabled.
 	if (AnyTransactionsInFlight() && FSanitizer::Mode != EMode::Disabled)
 	{
-		FSanitizer::FLock Sanitizer;
-		if (Sanitizer)
+		FSanitizer::FConflictingWrite::FPtr Err;
+		{
+			FSanitizer::FLock Sanitizer;
+			if (Sanitizer)
+			{
+				Err = Sanitizer->OpenWrite(Ptr, Size, __builtin_return_address(0));
+			}
+		}
+		if (Err)
 		{
-			Sanitizer->OpenWrite(Ptr, Size, __builtin_return_address(0));
+			// Must be called outside of the sanitizer lock.
+			Err->Report();
 		}
 	}
 }
