[AutoRTFM] Remove unused stats collection code.

These stats are very low-level and aren't used or tested anywhere.

#Synced-CL 50391842

diff --git a/Engine/Source/Runtime/AutoRTFM/Private/API.cpp b/Engine/Source/Runtime/AutoRTFM/Private/API.cpp
--- a/Engine/Source/Runtime/AutoRTFM/Private/API.cpp
+++ b/Engine/Source/Runtime/AutoRTFM/Private/API.cpp
@@ -5,6 +5,7 @@
 #include "AutoRTFM/Defines.h"
 #include "Context.h"
 #include "ExternAPI.h"
+#include "Toggles.h"
 #include "Transaction.h"
 
 #include <random>
diff --git a/Engine/Source/Runtime/AutoRTFM/Private/Context.cpp b/Engine/Source/Runtime/AutoRTFM/Private/Context.cpp
--- a/Engine/Source/Runtime/AutoRTFM/Private/Context.cpp
+++ b/Engine/Source/Runtime/AutoRTFM/Private/Context.cpp
@@ -11,7 +11,6 @@
 #include "CallNestInlines.h"
 #include "MetricsPriv.h"
 #include "StackRange.h"
-#include "Stats.h"
 #include "ToString.h"
 #include "Transaction.h"
 #include "TransactionInlines.h"
@@ -246,9 +245,6 @@ FTransaction* FContext::PushTransaction(FStackRange StackRange)
 
 	CurrentMaterializedTransaction = NewTransaction;
 
-	// Collect stats that we've got a new transaction.
-	Stats.Collect<EStatsKind::Transaction>();
-
 	return NewTransaction;
 }
 
diff --git a/Engine/Source/Runtime/AutoRTFM/Private/Stats.cpp b/Engine/Source/Runtime/AutoRTFM/Private/Stats.cpp
deleted file mode 100644
diff --git a/Engine/Source/Runtime/AutoRTFM/Private/Stats.h b/Engine/Source/Runtime/AutoRTFM/Private/Stats.h
deleted file mode 100644
diff --git a/Engine/Source/Runtime/AutoRTFM/Private/Transaction.cpp b/Engine/Source/Runtime/AutoRTFM/Private/Transaction.cpp
--- a/Engine/Source/Runtime/AutoRTFM/Private/Transaction.cpp
+++ b/Engine/Source/Runtime/AutoRTFM/Private/Transaction.cpp
@@ -26,7 +26,6 @@ void FTransaction::Resurrect(FContext* const InContext)
 	AUTORTFM_ASSERT(IsFresh());
 	Context = InContext;
 	Parent = nullptr;
-	StatDepth = 1;
 	bIsInAllocateFn = false;
 	ID = NextID++;
 }
@@ -73,15 +72,6 @@ void FTransaction::Initialize(FTransaction* InParent, FStackRange InStackRange)
 	Parent = InParent;
 	StackRange = InStackRange;
 	AUTORTFM_SANITIZER_ONLY(Sanitizer::StartTransaction(ID, StackRange));
-
-	if (Parent)
-	{
-		// For stats, record the nested depth of the transaction.
-		StatDepth = Parent->StatDepth + 1;
-	}
-
-	Stats.Collect<EStatsKind::AverageTransactionDepth>(StatDepth);
-	Stats.Collect<EStatsKind::MaximumTransactionDepth>(StatDepth);
 }
 
 bool FTransaction::IsFresh() const
@@ -96,9 +86,6 @@ void FTransaction::Abort()
 	AUTORTFM_ASSERT(IsActive());
 	AUTORTFM_ASSERT(Context->GetStatus() == EContextStatus::Aborting);
 
-	Stats.Collect<EStatsKind::Abort>();
-	CollectStats();
-
 	UpdatePeakMemoryUsageMetrics();
 	if (!Parent)
 	{
@@ -129,9 +116,6 @@ void FTransaction::Commit()
 	AUTORTFM_ASSERT(Context->GetCurrentTransaction() == this);
 	AUTORTFM_SANITIZER_ONLY(Sanitizer::CommitTransaction(ID));
 
-	Stats.Collect<EStatsKind::Commit>();
-	CollectStats();
-
 	if (Parent)
 	{
 		CommitNested();
diff --git a/Engine/Source/Runtime/AutoRTFM/Private/Transaction.h b/Engine/Source/Runtime/AutoRTFM/Private/Transaction.h
--- a/Engine/Source/Runtime/AutoRTFM/Private/Transaction.h
+++ b/Engine/Source/Runtime/AutoRTFM/Private/Transaction.h
@@ -10,7 +10,6 @@
 #include "IntrusivePool.h"
 #include "Stack.h"
 #include "StackRange.h"
-#include "Stats.h"
 #include "TaskArray.h"
 #include "Utils.h"
 #include "WriteLog.h"
@@ -95,8 +94,6 @@ private:
 	void CommitNested();
 	void CommitOuterNest();
 
-	void CollectStats() const;
-
 	bool ShouldRecordWrite(void* LogicalAddress) const;
 
 	FContext* Context = nullptr;  // null indicates the transaction is not active
@@ -137,7 +134,6 @@ private:
 	FHitSet HitSet;
 	FIntervalTree NewMemoryTracker;
 	FWriteLog WriteLog;
-	TStatStorage<uint64_t> StatDepth = 1;
 	FStackRange StackRange;
 	TransactionID ID = 0;
 	bool bIsInAllocateFn = false;
diff --git a/Engine/Source/Runtime/AutoRTFM/Private/TransactionInlines.h b/Engine/Source/Runtime/AutoRTFM/Private/TransactionInlines.h
--- a/Engine/Source/Runtime/AutoRTFM/Private/TransactionInlines.h
+++ b/Engine/Source/Runtime/AutoRTFM/Private/TransactionInlines.h
@@ -9,6 +9,7 @@
 #include "HitSet.h"
 #include "Sanitizer.h"
 #include "ScopedGuard.h"
+#include "Toggles.h"
 #include "Transaction.h"
 #include "Utils.h"
 #include "WriteLog.h"
@@ -72,7 +73,6 @@ AUTORTFM_NO_ASAN UE_AUTORTFM_FORCEINLINE void FTransaction::RecordWrite(
 
 	if (!ShouldRecordWrite(LogicalAddress))
 	{
-		Stats.Collect<EStatsKind::HitSetSkippedBecauseOfStackLocalMemory>();
 		return;
 	}
 
@@ -85,21 +85,15 @@ AUTORTFM_NO_ASAN UE_AUTORTFM_FORCEINLINE void FTransaction::RecordWrite(
 
 		if (HitSet.FindOrTryInsert(HitSetEntry))
 		{
-			Stats.Collect<EStatsKind::HitSetHit>();
 			return;
 		}
-
-		Stats.Collect<EStatsKind::HitSetMiss>();
 	}
 
 	if (NewMemoryTracker.Contains(LogicalAddress, Size))
 	{
-		Stats.Collect<EStatsKind::NewMemoryTrackerHit>();
 		return;
 	}
 
-	Stats.Collect<EStatsKind::NewMemoryTrackerMiss>();
-
 	FWriteLogEntry Entry{.LogicalAddress = static_cast<std::byte*>(LogicalAddress),
 		.Data = static_cast<std::byte*>(LogicalAddress),
 		.Size = Size,
@@ -137,7 +131,6 @@ AUTORTFM_NO_ASAN UE_AUTORTFM_FORCEINLINE void FTransaction::RecordWrite(void* Lo
 	static_assert(SIZE <= 8);
 	if (!ShouldRecordWrite(LogicalAddress))
 	{
-		Stats.Collect<EStatsKind::HitSetSkippedBecauseOfStackLocalMemory>();
 		return;
 	}
 
@@ -148,7 +141,6 @@ AUTORTFM_NO_ASAN UE_AUTORTFM_FORCEINLINE void FTransaction::RecordWrite(void* Lo
 	switch (HitSet.FindOrTryInsertNoResize(Entry))
 	{
 		case FHitSet::EInsertResult::Exists:
-			Stats.Collect<EStatsKind::HitSetHit>();
 			return;
 		case FHitSet::EInsertResult::Inserted:
 			AUTORTFM_MUST_TAIL return FTransaction::RecordWriteInsertedSlow<SIZE>(LogicalAddress);
@@ -166,12 +158,9 @@ AUTORTFM_NO_ASAN UE_AUTORTFM_FORCENOINLINE void FTransaction::RecordWriteNotInse
 
 	if (HitSet.FindOrTryInsert(Entry))
 	{
-		Stats.Collect<EStatsKind::HitSetHit>();
 		return;
 	}
 
-	Stats.Collect<EStatsKind::HitSetMiss>();
-
 	return RecordWriteInsertedSlow<SIZE>(LogicalAddress);
 }
 
@@ -180,12 +169,9 @@ AUTORTFM_NO_ASAN UE_AUTORTFM_FORCENOINLINE void FTransaction::RecordWriteInserte
 {
 	if (NewMemoryTracker.Contains(LogicalAddress, SIZE))
 	{
-		Stats.Collect<EStatsKind::NewMemoryTrackerHit>();
 		return;
 	}
 
-	Stats.Collect<EStatsKind::NewMemoryTrackerMiss>();
-
 	AutoRTFM::TScopedGuard<bool> RecursionGuard(bIsInAllocateFn, true);
 
 	WriteLog.PushSmall<SIZE>(static_cast<std::byte*>(LogicalAddress), static_cast<std::byte*>(LogicalAddress));
@@ -303,27 +289,6 @@ UE_AUTORTFM_FORCEINLINE void FTransaction::PopAllDeferUntilAbortHandlers(const v
 	DeferredPopAllOnAbortHandlers.Push(Key);
 }
 
-UE_AUTORTFM_FORCEINLINE void FTransaction::CollectStats() const
-{
-	Stats.Collect<EStatsKind::AverageWriteLogEntries>(WriteLog.Num());
-	Stats.Collect<EStatsKind::MaximumWriteLogEntries>(WriteLog.Num());
-
-	Stats.Collect<EStatsKind::AverageWriteLogBytes>(WriteLog.TotalSize());
-	Stats.Collect<EStatsKind::MaximumWriteLogBytes>(WriteLog.TotalSize());
-
-	Stats.Collect<EStatsKind::AverageCommitTasks>(CommitTasks.Num());
-	Stats.Collect<EStatsKind::MaximumCommitTasks>(CommitTasks.Num());
-
-	Stats.Collect<EStatsKind::AveragePreAbortTasks>(PreAbortTasks.Num());
-	Stats.Collect<EStatsKind::MaximumPreAbortTasks>(PreAbortTasks.Num());
-
-	Stats.Collect<EStatsKind::AverageAbortTasks>(AbortTasks.Num());
-	Stats.Collect<EStatsKind::MaximumAbortTasks>(AbortTasks.Num());
-
-	Stats.Collect<EStatsKind::AverageHitSetSize>(HitSet.GetCount());
-	Stats.Collect<EStatsKind::AverageHitSetCapacity>(HitSet.GetCapacity());
-}
-
 }  // namespace AutoRTFM
 
 #undef AUTORTFM_CHECK_ASAN_FAKE_STACKS
