Add missing `UE_AUTORTFM_API` on exported symbols.

#Synced-CL 50451909

diff --git a/Engine/Source/Runtime/AutoRTFM/Private/MemoryStats.h b/Engine/Source/Runtime/AutoRTFM/Private/MemoryStats.h
--- a/Engine/Source/Runtime/AutoRTFM/Private/MemoryStats.h
+++ b/Engine/Source/Runtime/AutoRTFM/Private/MemoryStats.h
@@ -45,7 +45,7 @@ struct FMemoryStats
 
 private:
 #if AUTORTFM_CHECK_FOR_ALLOCATION_UNDERFLOW
-	static void FatalErrorBytesAllocatedUnderflow();  // Not inline to break circular header dependencies.
+	UE_AUTORTFM_API static void FatalErrorBytesAllocatedUnderflow();  // Not inline to break circular header dependencies.
 #endif
 
 	UE_AUTORTFM_API static std::atomic<size_t> BytesAllocated;
diff --git a/Engine/Source/Runtime/AutoRTFM/Private/Utils.h b/Engine/Source/Runtime/AutoRTFM/Private/Utils.h
--- a/Engine/Source/Runtime/AutoRTFM/Private/Utils.h
+++ b/Engine/Source/Runtime/AutoRTFM/Private/Utils.h
@@ -51,7 +51,7 @@ std::string GetFunctionDescription(TReturnType (*FunctionPtr)(TParameterTypes...
 // the cost of calling this function (which is super unlikely to
 // be called because it is the asserting true codepath, EG. only
 // fatal process destroying explosions happen here).
-[[clang::preserve_most, noreturn]] UE_AUTORTFM_FORCENOINLINE void DoAssert(void (*Logger)());
+[[clang::preserve_most, noreturn]] UE_AUTORTFM_API UE_AUTORTFM_FORCENOINLINE void DoAssert(void (*Logger)());
 
 // We use a special calling convention for this call to minimize
 // the cost of calling this function (which is super unlikely to
