Allow ubsan reports in the closed.

#Synced-CL 50530424

diff --git a/Engine/Source/Runtime/AutoRTFM/AutoRTFM.Build.cs b/Engine/Source/Runtime/AutoRTFM/AutoRTFM.Build.cs
--- a/Engine/Source/Runtime/AutoRTFM/AutoRTFM.Build.cs
+++ b/Engine/Source/Runtime/AutoRTFM/AutoRTFM.Build.cs
@@ -22,18 +22,31 @@ namespace UnrealBuildTool.Rules
 			if (Target.Platform == UnrealTargetPlatform.Win64)
 			{
 				AutoRTFMExternalMappingFiles.Add("Runtime/AutoRTFM/Public/StdLib.Windows.aem");
+
 				if (Target.WindowsPlatform.bEnableAddressSanitizer)
 				{
 					AutoRTFMExternalMappingFiles.Add("Runtime/AutoRTFM/Public/ASAN.aem");
 				}
+
+				if (Target.WindowsPlatform.bEnableUndefinedBehaviorSanitizer)
+				{
+					AutoRTFMExternalMappingFiles.Add("Runtime/AutoRTFM/Public/UBSAN.aem");
+				}
 			}
 			if (Target.Platform == UnrealTargetPlatform.Linux || Target.Platform == UnrealTargetPlatform.LinuxArm64)
 			{
 				AutoRTFMExternalMappingFiles.Add("Runtime/AutoRTFM/Public/StdLib.Linux.aem");
+
 				if (Target.LinuxPlatform.bEnableAddressSanitizer)
 				{
 					AutoRTFMExternalMappingFiles.Add("Runtime/AutoRTFM/Public/ASAN.aem");
 				}
+
+				if (Target.LinuxPlatform.bEnableUndefinedBehaviorSanitizer)
+				{
+					AutoRTFMExternalMappingFiles.Add("Runtime/AutoRTFM/Public/UBSAN.aem");
+				}
+
 				if (Target.StaticAllocator == StaticAllocatorType.Ansi)
 				{
 					AutoRTFMExternalMappingFiles.Add("Runtime/AutoRTFM/Public/NewDelete.Linux.aem");
