[AutoRTFMSan] Fix infinite recursion in debug builds

When built with debug, std::unique_ptr has non-inlined, header implementations that when instrumented, emit calls to autortfm_sanitizer_open_write.
Using std::unique_ptr in autortfm_sanitizer_open_write causes the function to recurse until stack overflowing.

#Synced-CL 50428254

diff --git a/Engine/Source/Runtime/AutoRTFM/Private/Sanitizer.cpp b/Engine/Source/Runtime/AutoRTFM/Private/Sanitizer.cpp
--- a/Engine/Source/Runtime/AutoRTFM/Private/Sanitizer.cpp
+++ b/Engine/Source/Runtime/AutoRTFM/Private/Sanitizer.cpp
@@ -50,16 +50,6 @@ struct AUTORTFM_DISABLE_NO_SANITIZE FOSAllocator
 	{
 		free(Ptr);
 	}
-
-	// Deleter that can be used as the default destruction policy of std::unique_ptr.
-	struct FDeleter
-	{
-		template <typename T>
-		void operator()(T* Ptr) const
-		{
-			Delete<FOSAllocator>(Ptr);
-		}
-	};
 };
 static_assert(Allocator<FOSAllocator>);
 
@@ -69,8 +59,6 @@ struct AUTORTFM_DISABLE_NO_SANITIZE FSanitizer
 
 	struct FConflictingWrite
 	{
-		using FPtr = std::unique_ptr<FConflictingWrite, FOSAllocator::FDeleter>;
-
 		autortfm_log_severity const Severity;
 		void* const Ptr;
 		size_t const Size;
@@ -101,7 +89,7 @@ struct AUTORTFM_DISABLE_NO_SANITIZE FSanitizer
 	void CommitTransaction(TransactionID ID);
 	void AbortTransaction(TransactionID ID);
 
-	[[nodiscard]] FConflictingWrite::FPtr OpenWrite(void* Ptr, size_t Size, void* ProgramCounter);
+	[[nodiscard]] FConflictingWrite* OpenWrite(void* Ptr, size_t Size, void* ProgramCounter);
 	void ClosedWrite(TransactionID ID, void* Ptr, size_t Size, void* ProgramCounter);
 
 private:
@@ -261,7 +249,7 @@ void FSanitizer::AbortTransaction(TransactionID ID)
 	PopTransaction(ID);
 }
 
-[[nodiscard]] FSanitizer::FConflictingWrite::FPtr FSanitizer::OpenWrite(void* Ptr, size_t Size, void* ProgramCounter)
+[[nodiscard]] FSanitizer::FConflictingWrite* FSanitizer::OpenWrite(void* Ptr, size_t Size, void* ProgramCounter)
 {
 	if (!CurrentTransaction || CurrentTransaction->DisableCount)
 	{
@@ -290,8 +278,8 @@ void FSanitizer::AbortTransaction(TransactionID ID)
 
 			NumIssues++;
 
-			return FConflictingWrite::FPtr(New<FOSAllocator, FConflictingWrite>(
-				Severity, Ptr, Size, CallstackTree.Get(OpenWriteCallstack), CallstackTree.Get(*ClosedWriteCallstack)));
+			return New<FOSAllocator, FConflictingWrite>(
+				Severity, Ptr, Size, CallstackTree.Get(OpenWriteCallstack), CallstackTree.Get(*ClosedWriteCallstack));
 		}
 	}
 
@@ -435,7 +423,7 @@ extern "C" AUTORTFM_DISABLE_NO_SANITIZE UE_AUTORTFM_FORCENOINLINE UE_AUTORTFM_AP
 	// if the sanitizer is disabled.
 	if (AnyTransactionsInFlight() && FSanitizer::Mode != EMode::Disabled)
 	{
-		FSanitizer::FConflictingWrite::FPtr Err;
+		FSanitizer::FConflictingWrite* Err = nullptr;
 		{
 			FSanitizer::FLock Sanitizer;
 			if (Sanitizer)
@@ -443,10 +431,12 @@ extern "C" AUTORTFM_DISABLE_NO_SANITIZE UE_AUTORTFM_FORCENOINLINE UE_AUTORTFM_AP
 				Err = Sanitizer->OpenWrite(Ptr, Size, __builtin_return_address(0));
 			}
 		}
+
 		if (Err)
 		{
 			// Must be called outside of the sanitizer lock.
 			Err->Report();
+			Delete<FOSAllocator>(Err);
 		}
 	}
 }
