[AutoRTFM] Add new AutoRTFM::EContextStatus::Unwinding state

Use this new context status between cascading abort / cascading retry and reaching the outer transaction, instead of EContextStatus::OnTrack with no transactions on the stack.
Fixes bad assumptions in mutexes that look at the context status.

Also:
* Rename the `AbortedByTransactDuring` statuses. They're not aborts as the transaction never started.
* Remove the `AbortedByTransactDuring` statuses from ETransactionStatus. They're not a transaction status, so never used.

#Synced-CL 50590394

diff --git a/Engine/Source/Programs/AutoRTFMTests/Private/API.cpp b/Engine/Source/Programs/AutoRTFMTests/Private/API.cpp
--- a/Engine/Source/Programs/AutoRTFMTests/Private/API.cpp
+++ b/Engine/Source/Programs/AutoRTFMTests/Private/API.cpp
@@ -1515,10 +1515,12 @@ TEST_CASE("API.CascadingAbortTransaction_WithinUnscopedTransaction")
 		{
 			AutoRTFM::StartTransaction();
 			AutoRTFM::CascadingAbortTransaction();
+			REQUIRE(AutoRTFM::ContextStatus() == AutoRTFM::EContextStatus::Unwinding);
 		});
 
 		FAIL("unreachable");
 	});
+	REQUIRE(AutoRTFM::ContextStatus() == AutoRTFM::EContextStatus::Idle);
 }
 
 TEST_CASE("API.CascadingRetryTransaction")
@@ -1618,6 +1620,38 @@ TEST_CASE("API.CascadingRetryTransaction")
 		REQUIRE(CallCount == 1);
 	}
 
+	SECTION("Open")
+	{
+		AutoRTFM::Testing::Commit([&]
+		{
+			bTransactionalWrite = true;
+			if (CallCount == 0)
+			{
+				AutoRTFM::OnRetry([&]
+				{
+					// The write above would have been undone.
+					REQUIRE(!bTransactionalWrite);
+					CallCount++;
+				});
+				AutoRTFM::Open([&]
+				{
+					AutoRTFM::ETransactionStatus Status = AutoRTFM::Close([&]
+					{
+						AutoRTFM::Open([&]
+						{
+							AutoRTFM::CascadingRetryTransaction();
+							REQUIRE(AutoRTFM::ContextStatus() == AutoRTFM::EContextStatus::Unwinding);
+						});
+					});
+					REQUIRE(AutoRTFM::ContextStatus() == AutoRTFM::EContextStatus::Unwinding);
+					REQUIRE(Status == AutoRTFM::ETransactionStatus::AbortedByCascadingRetry);
+				});
+			}
+		});
+		REQUIRE(bTransactionalWrite);
+		REQUIRE(CallCount == 1);
+	}
+
 	SECTION("Status query during retry")
 	{
 		AutoRTFM::Testing::Commit([&]
diff --git a/Engine/Source/Programs/AutoRTFMTests/Private/Abort.cpp b/Engine/Source/Programs/AutoRTFMTests/Private/Abort.cpp
--- a/Engine/Source/Programs/AutoRTFMTests/Private/Abort.cpp
+++ b/Engine/Source/Programs/AutoRTFMTests/Private/Abort.cpp
@@ -117,7 +117,7 @@ TEST_CASE("Abort.TransactionInOnCommit")
 				bDidSomething = true;
 			});
 
-			REQUIRE(AutoRTFM::ETransactionResult::AbortedByTransactDuringCommit == Result);
+			REQUIRE(AutoRTFM::ETransactionResult::RejectedTransactDuringCommit == Result);
 			REQUIRE(false == bDidSomething);
 		});
 	});
@@ -137,7 +137,7 @@ TEST_CASE("Abort.TransactionInOnAbort")
 				bDidSomething = true;
 			});
 
-			REQUIRE(AutoRTFM::ETransactionResult::AbortedByTransactDuringAbort == InnerResult);
+			REQUIRE(AutoRTFM::ETransactionResult::RejectedTransactDuringAbort == InnerResult);
 			REQUIRE(false == bDidSomething);
 		});
 
@@ -160,7 +160,7 @@ TEST_CASE("Abort.TransactionInOnPreAbort")
 				bDidSomething = true;
 			});
 
-			REQUIRE(AutoRTFM::ETransactionResult::AbortedByTransactDuringAbort == InnerResult);
+			REQUIRE(AutoRTFM::ETransactionResult::RejectedTransactDuringAbort == InnerResult);
 			REQUIRE(false == bDidSomething);
 		});
 
@@ -183,7 +183,7 @@ TEST_CASE("Abort.TransactionInAbortedOnComplete")
 				bDidSomething = true;
 			});
 
-			REQUIRE(InnerResult == AutoRTFM::ETransactionResult::AbortedByTransactDuringCompletion);
+			REQUIRE(InnerResult == AutoRTFM::ETransactionResult::RejectedTransactDuringCompletion);
 			REQUIRE(!bDidSomething);
 		});
 
@@ -210,7 +210,7 @@ TEST_CASE("Abort.TransactionInOnComplete")
 			});
 
 			REQUIRE(!bDidSomething);
-			REQUIRE(InnerResult == AutoRTFM::ETransactionResult::AbortedByTransactDuringCompletion);
+			REQUIRE(InnerResult == AutoRTFM::ETransactionResult::RejectedTransactDuringCompletion);
 		});
 	});
 	REQUIRE(CallCount == 1);
@@ -235,7 +235,7 @@ TEST_CASE("Abort.TransactionInOnRetry")
 			});
 
 			REQUIRE(!bDidSomething);
-			REQUIRE(InnerResult == AutoRTFM::ETransactionResult::AbortedByTransactDuringRetry);
+			REQUIRE(InnerResult == AutoRTFM::ETransactionResult::RejectedTransactDuringRetry);
 		});
 		
 		if (CallCount == 0)
diff --git a/Engine/Source/Runtime/AutoRTFM/Private/Context.cpp b/Engine/Source/Runtime/AutoRTFM/Private/Context.cpp
--- a/Engine/Source/Runtime/AutoRTFM/Private/Context.cpp
+++ b/Engine/Source/Runtime/AutoRTFM/Private/Context.cpp
@@ -151,6 +151,8 @@ void FContext::AbortTransaction(ETransactionStatus Reason)
 			CurrentMaterializedTransaction->Abort();
 			PopTransaction();
 		}
+
+		Status = EContextStatus::Unwinding;
 	}
 	else
 	{
@@ -171,9 +173,9 @@ void FContext::AbortTransaction(ETransactionStatus Reason)
 			CurrentMaterializedTransaction->Abort();
 			PopTransaction();
 		}
-	}
 
-	Status = EContextStatus::OnTrack;
+		Status = EContextStatus::OnTrack;
+	}
 }
 
 void FContext::SetCallNestTransactionCascadingStatus(ETransactionStatus TransactionStatus)
@@ -335,14 +337,16 @@ ETransactionResult FContext::Transact(void (*UninstrumentedFunction)(void*), voi
 {
 	switch (Status)
 	{
+		[[unlikely]] case EContextStatus::Unwinding:
+			return ETransactionResult::RejectedTransactDuringUnwind;
 		[[unlikely]] case EContextStatus::Committing:
-			return ETransactionResult::AbortedByTransactDuringCommit;
+			return ETransactionResult::RejectedTransactDuringCommit;
 		[[unlikely]] case EContextStatus::Aborting:
-			return ETransactionResult::AbortedByTransactDuringAbort;
+			return ETransactionResult::RejectedTransactDuringAbort;
 		[[unlikely]] case EContextStatus::Retrying:
-			return ETransactionResult::AbortedByTransactDuringRetry;
+			return ETransactionResult::RejectedTransactDuringRetry;
 		[[unlikely]] case EContextStatus::Completing:
-			return ETransactionResult::AbortedByTransactDuringCompletion;
+			return ETransactionResult::RejectedTransactDuringCompletion;
 
 		case EContextStatus::Idle:
 		case EContextStatus::OnTrack:
@@ -491,7 +495,7 @@ ETransactionResult FContext::Transact(void (*UninstrumentedFunction)(void*), voi
 
 void FContext::EndScopedTransaction()
 {
-	AUTORTFM_ASSERT(Status == EContextStatus::OnTrack);
+	AUTORTFM_ASSERT(Status == EContextStatus::OnTrack || Status == EContextStatus::Unwinding);
 	if (CurrentMaterializedTransaction == nullptr)
 	{
 		// No more transactions in flight.
@@ -503,11 +507,13 @@ void FContext::EndScopedTransaction()
 
 void FContext::CallCompletionTasks()
 {
-	AUTORTFM_ASSERT(Status == EContextStatus::OnTrack);
 	AUTORTFM_ASSERT(CurrentMaterializedTransaction == nullptr);
-	Status = EContextStatus::Completing;
-	CompletionTasks.RemoveEachForward([](TTask<void()>& Task) { Task(); });
-	Status = EContextStatus::OnTrack;
+	AUTORTFM_ASSERT(Status == EContextStatus::OnTrack || Status == EContextStatus::Unwinding);
+
+	{
+		AUTORTFM_SCOPED_ASSIGNMENT(Status, EContextStatus::Completing);
+		CompletionTasks.RemoveEachForward([](TTask<void()>& Task) { Task(); });
+	}
 }
 
 void FContext::Reset()
diff --git a/Engine/Source/Runtime/AutoRTFM/Private/ToString.h b/Engine/Source/Runtime/AutoRTFM/Private/ToString.h
--- a/Engine/Source/Runtime/AutoRTFM/Private/ToString.h
+++ b/Engine/Source/Runtime/AutoRTFM/Private/ToString.h
@@ -17,6 +17,8 @@ inline const char* ToString(EContextStatus Status)
 			return "Idle";
 		case EContextStatus::OnTrack:
 			return "OnTrack";
+		case EContextStatus::Unwinding:
+			return "Unwinding";
 		case EContextStatus::Committing:
 			return "Committing";
 		case EContextStatus::Aborting:
@@ -42,14 +44,6 @@ inline const char* ToString(ETransactionStatus Status)
 			return "AbortedByRequest";
 		case ETransactionStatus::AbortedByLanguage:
 			return "AbortedByLanguage";
-		case ETransactionStatus::AbortedByTransactDuringCommit:
-			return "AbortedByTransactDuringCommit";
-		case ETransactionStatus::AbortedByTransactDuringAbort:
-			return "AbortedByTransactDuringAbort";
-		case ETransactionStatus::AbortedByTransactDuringRetry:
-			return "AbortedByTransactDuringRetry";
-		case ETransactionStatus::AbortedByTransactDuringCompletion:
-			return "AbortedByTransactDuringCompletion";
 		case ETransactionStatus::AbortedByCascadingAbort:
 			return "AbortedByCascadingAbort";
 		case ETransactionStatus::AbortedByCascadingRetry:
@@ -67,14 +61,16 @@ inline const char* ToString(ETransactionResult Result)
 			return "AbortedByRequest";
 		case ETransactionResult::AbortedByLanguage:
 			return "AbortedByLanguage";
-		case ETransactionResult::AbortedByTransactDuringCommit:
-			return "AbortedByTransactDuringCommit";
-		case ETransactionResult::AbortedByTransactDuringAbort:
-			return "AbortedByTransactDuringAbort";
-		case ETransactionResult::AbortedByTransactDuringRetry:
-			return "AbortedByTransactDuringRetry";
-		case ETransactionResult::AbortedByTransactDuringCompletion:
-			return "AbortedByTransactDuringCompletion";
+		case ETransactionResult::RejectedTransactDuringUnwind:
+			return "RejectedTransactDuringUnwind";
+		case ETransactionResult::RejectedTransactDuringCommit:
+			return "RejectedTransactDuringCommit";
+		case ETransactionResult::RejectedTransactDuringAbort:
+			return "RejectedTransactDuringAbort";
+		case ETransactionResult::RejectedTransactDuringRetry:
+			return "RejectedTransactDuringRetry";
+		case ETransactionResult::RejectedTransactDuringCompletion:
+			return "RejectedTransactDuringCompletion";
 		case ETransactionResult::AbortedByCascadingAbort:
 			return "AbortedByCascadingAbort";
 	}
diff --git a/Engine/Source/Runtime/AutoRTFM/Private/Utils.h b/Engine/Source/Runtime/AutoRTFM/Private/Utils.h
--- a/Engine/Source/Runtime/AutoRTFM/Private/Utils.h
+++ b/Engine/Source/Runtime/AutoRTFM/Private/Utils.h
@@ -173,6 +173,36 @@ inline constexpr T Log2(T Value)
 	return Value > 0 ? NumBits - __builtin_clzg(Value) - 1 : 0;
 }
 
+// RAII helper to assign a value to a variable for the lifetime of the
+// FScopedAssignment, restoring the original value when the FScopedAssignment is
+// destructed.
+template <typename T>
+struct TScopedAssignment
+{
+	inline TScopedAssignment(T& Variable, T NewValue) : Variable(Variable), OldValue(Variable)
+	{
+		Variable = NewValue;
+	}
+
+	inline ~TScopedAssignment()
+	{
+		Variable = OldValue;
+	}
+
+private:
+	TScopedAssignment(const TScopedAssignment&) = delete;
+	TScopedAssignment(TScopedAssignment&&) = delete;
+	TScopedAssignment& operator=(const TScopedAssignment&) = delete;
+	TScopedAssignment& operator=(TScopedAssignment&&) = delete;
+	T& Variable;
+	T OldValue;
+};
+
+// Assigns Value to Variable, and restores the original value of Variable at the
+// end of the current scope.
+#define AUTORTFM_SCOPED_ASSIGNMENT(Variable, Value) \
+	TScopedAssignment UE_AUTORTFM_CONCAT(ScopedAssignment, __COUNTER__){Variable, Value};
+
 }  // namespace AutoRTFM
 
 #define AUTORTFM_REQUIRE_SEMICOLON static_assert(true)
diff --git a/Engine/Source/Runtime/AutoRTFM/Public/AutoRTFM.h b/Engine/Source/Runtime/AutoRTFM/Public/AutoRTFM.h
--- a/Engine/Source/Runtime/AutoRTFM/Public/AutoRTFM.h
+++ b/Engine/Source/Runtime/AutoRTFM/Public/AutoRTFM.h
@@ -84,24 +84,28 @@ enum class ETransactionResult
 	// The transaction aborted because of unhandled constructs in the code (atomics, unhandled function calls, etc).
 	AbortedByLanguage = autortfm_transaction_aborted_by_language,
 
-	// The transaction aborted because a new transaction nest was attempted (via OnCommit) while the
-	// current transaction was being committed.
-	AbortedByTransactDuringCommit = autortfm_transaction_aborted_by_transact_during_commit,
+	// The transaction aborted because of an explicit call to CascadingAbortTransaction.
+	AbortedByCascadingAbort = autortfm_transaction_aborted_by_cascading_abort,
 
-	// The transaction aborted because a new transaction nest was attempted (via OnAbort) while the
-	// current transaction was being aborted.
-	AbortedByTransactDuringAbort = autortfm_transaction_aborted_by_transact_during_abort,
+	// A transaction was rejected because a new transaction nest was attempted
+	// after a cascading abort or cascading retry.
+	RejectedTransactDuringUnwind = autortfm_transaction_rejected_during_unwind,
 
-	// The transaction aborted because a new transaction nest was attempted (via OnRetry) while the
-	// current transaction was being retried.
-	AbortedByTransactDuringRetry = autortfm_transaction_aborted_by_transact_during_retry,
+	// A transaction was rejected because a new transaction nest was attempted
+	// (via OnCommit) while the current transaction was being committed.
+	RejectedTransactDuringCommit = autortfm_transaction_rejected_during_commit,
 
-	// The transaction aborted because a new transaction nest was attempted (via OnComplete) while the
-	// current transaction was being retried.
-	AbortedByTransactDuringCompletion = autortfm_transaction_aborted_by_transact_during_complete,
+	// A transaction was rejected because a new transaction nest was attempted
+	// (via OnAbort) while the current transaction was being aborted.
+	RejectedTransactDuringAbort = autortfm_transaction_rejected_during_abort,
 
-	// The transaction aborted because of an explicit call to CascadingAbortTransaction.
-	AbortedByCascadingAbort = autortfm_transaction_aborted_by_cascading_abort,
+	// A transaction was rejected because a new transaction nest was attempted
+	// (via OnRetry) while the current transaction was being retried.
+	RejectedTransactDuringRetry = autortfm_transaction_rejected_during_retry,
+
+	// A transaction was rejected because a new transaction nest was attempted
+	// (via OnComplete) while the current transaction was being retried.
+	RejectedTransactDuringCompletion = autortfm_transaction_rejected_during_complete,
 };
 
 // Returns true if the given ETransactionResult represents a transaction that
@@ -116,18 +120,18 @@ inline bool HasAborted(ETransactionResult State)
 
 		case ETransactionResult::AbortedByRequest:
 		case ETransactionResult::AbortedByLanguage:
-		case ETransactionResult::AbortedByTransactDuringCommit:
-		case ETransactionResult::AbortedByTransactDuringAbort:
-		case ETransactionResult::AbortedByTransactDuringRetry:
-		case ETransactionResult::AbortedByTransactDuringCompletion:
 		case ETransactionResult::AbortedByCascadingAbort:
+		case ETransactionResult::RejectedTransactDuringUnwind:
+		case ETransactionResult::RejectedTransactDuringCommit:
+		case ETransactionResult::RejectedTransactDuringAbort:
+		case ETransactionResult::RejectedTransactDuringRetry:
+		case ETransactionResult::RejectedTransactDuringCompletion:
 			break;
 	}
 	return true;
 }
 
 // The transaction status provides information on how a transaction is currently executing or completed.
-// ETransactionStatus is a superset of ETransactionResult.
 enum class ETransactionStatus
 {
 	// The transaction aborted because of an explicit call to AbortTransaction.
@@ -143,22 +147,6 @@ enum class ETransactionStatus
 	// The transaction aborted because of unhandled constructs in the code (atomics, unhandled function calls, etc).
 	AbortedByLanguage = autortfm_transaction_aborted_by_language,
 
-	// The transaction aborted because a new transaction nest was attempted (via OnCommit) while the
-	// current transaction was being committed.
-	AbortedByTransactDuringCommit = autortfm_transaction_aborted_by_transact_during_commit,
-
-	// The transaction aborted because a new transaction nest was attempted (via OnAbort) while the
-	// current transaction was being aborted.
-	AbortedByTransactDuringAbort = autortfm_transaction_aborted_by_transact_during_abort,
-
-	// The transaction aborted because a new transaction nest was attempted (via OnRetry) while the
-	// current transaction was being retried.
-	AbortedByTransactDuringRetry = autortfm_transaction_aborted_by_transact_during_retry,
-
-	// The transaction aborted because a new transaction nest was attempted (via OnComplete) while the
-	// current transaction was being retried.
-	AbortedByTransactDuringCompletion = autortfm_transaction_aborted_by_transact_during_complete,
-
 	// The transaction aborted because of an explicit call to CascadingAbortTransaction.
 	AbortedByCascadingAbort = autortfm_transaction_aborted_by_cascading_abort,
 
@@ -178,10 +166,6 @@ inline bool HasAborted(ETransactionStatus State)
 
 		case ETransactionStatus::AbortedByRequest:
 		case ETransactionStatus::AbortedByLanguage:
-		case ETransactionStatus::AbortedByTransactDuringCommit:
-		case ETransactionStatus::AbortedByTransactDuringAbort:
-		case ETransactionStatus::AbortedByTransactDuringRetry:
-		case ETransactionStatus::AbortedByTransactDuringCompletion:
 		case ETransactionStatus::AbortedByCascadingAbort:
 		case ETransactionStatus::AbortedByCascadingRetry:
 			break;
@@ -203,10 +187,6 @@ inline bool IsCascading(ETransactionStatus Status)
 		case ETransactionStatus::Committed:
 		case ETransactionStatus::AbortedByRequest:
 		case ETransactionStatus::AbortedByLanguage:
-		case ETransactionStatus::AbortedByTransactDuringCommit:
-		case ETransactionStatus::AbortedByTransactDuringAbort:
-		case ETransactionStatus::AbortedByTransactDuringRetry:
-		case ETransactionStatus::AbortedByTransactDuringCompletion:
 			break;
 	}
 	return false;
@@ -222,6 +202,10 @@ enum class EContextStatus : uint8_t
 	// An OnTrack status means we are in transactional code.
 	OnTrack = autortfm_status_on_track,
 
+	// An Unwinding status means we've performed a cascading abort / retry and
+	// are now in the process of unwinding the stack to the outermost Transact()
+	Unwinding = autortfm_status_unwinding,
+
 	// We are currently running the on-commit handlers for a transaction.
 	Committing = autortfm_status_committing,
 
@@ -441,24 +425,25 @@ UE_AUTORTFM_CRITICAL_INLINE_ALWAYS bool IsRetrying()
 	return autortfm_is_context_status(autortfm_status_retrying);
 }
 
-// Returns true when inside an on-abort, on-commit, on-complete, on-retry handler
-// or inside a static initializer.
+// Returns true when inside an on-abort, on-commit, on-complete, on-retry handler,
+// unwinding from a cascading abort / cascading retry, or inside a static initializer.
 // TODO(SOL-8967): This function name is misleading. Rename to something sensible.
 UE_AUTORTFM_CRITICAL_INLINE_ALWAYS bool IsCommittingOrAborting()
 {
 	switch (ContextStatus())
 	{
+		case AutoRTFM::EContextStatus::Idle:
+		case AutoRTFM::EContextStatus::OnTrack:
+			return false;
 		case AutoRTFM::EContextStatus::Committing:
 		case AutoRTFM::EContextStatus::Aborting:
 		case AutoRTFM::EContextStatus::Completing:
 		case AutoRTFM::EContextStatus::Retrying:
 		case AutoRTFM::EContextStatus::InStaticLocalInitializer:
-			return true;
-		default:
-		case AutoRTFM::EContextStatus::Idle:
-		case AutoRTFM::EContextStatus::OnTrack:
-			return false;
+		case AutoRTFM::EContextStatus::Unwinding:
+			break;
 	}
+	return true;
 }
 
 // Returns true if the passed-in pointer is on the stack of the currently-executing transaction.
diff --git a/Engine/Source/Runtime/AutoRTFM/Public/AutoRTFM/CAPI.h b/Engine/Source/Runtime/AutoRTFM/Public/AutoRTFM/CAPI.h
--- a/Engine/Source/Runtime/AutoRTFM/Public/AutoRTFM/CAPI.h
+++ b/Engine/Source/Runtime/AutoRTFM/Public/AutoRTFM/CAPI.h
@@ -49,10 +49,11 @@ extern "C"
 		autortfm_transaction_committed,
 		autortfm_transaction_aborted_by_request,
 		autortfm_transaction_aborted_by_language,
-		autortfm_transaction_aborted_by_transact_during_commit,
-		autortfm_transaction_aborted_by_transact_during_abort,
-		autortfm_transaction_aborted_by_transact_during_retry,
-		autortfm_transaction_aborted_by_transact_during_complete,
+		autortfm_transaction_rejected_during_unwind,
+		autortfm_transaction_rejected_during_commit,
+		autortfm_transaction_rejected_during_abort,
+		autortfm_transaction_rejected_during_retry,
+		autortfm_transaction_rejected_during_complete,
 		autortfm_transaction_aborted_by_cascading_abort,
 		autortfm_transaction_aborted_by_cascading_retry,
 	} autortfm_transaction_status;
diff --git a/Engine/Source/Runtime/AutoRTFM/Public/AutoRTFM/Constants.h b/Engine/Source/Runtime/AutoRTFM/Public/AutoRTFM/Constants.h
--- a/Engine/Source/Runtime/AutoRTFM/Public/AutoRTFM/Constants.h
+++ b/Engine/Source/Runtime/AutoRTFM/Public/AutoRTFM/Constants.h
@@ -16,6 +16,7 @@ extern "C"
 	{
 		autortfm_status_idle = 0,
 		autortfm_status_on_track,
+		autortfm_status_unwinding,
 		autortfm_status_committing,
 		autortfm_status_aborting,
 		autortfm_status_retrying,
