Fix the `-LTO -LTCG` issue with merge modules and modular AutoRTFM test targets.

There are two bugs we've had to 'fix' (workaround) in here:

- Firstly that our explicit AutoRTFM External Mappings for internal symbols don't play nice with relocations (jira SOL-9048 covers that now).
- Secondly that having our relocations and using section prefixes causes relocations too (similar errors to having open <-> closed functions in separate sections). Workaround that by disabling PGO/LTOs ability to move functions into section prefixes (jira SOL-9041 covers that now).

#Synced-CL 50768071

diff --git a/Engine/Source/Programs/UnrealBuildTool/Platform/Linux/UEBuildLinux.cs b/Engine/Source/Programs/UnrealBuildTool/Platform/Linux/UEBuildLinux.cs
--- a/Engine/Source/Programs/UnrealBuildTool/Platform/Linux/UEBuildLinux.cs
+++ b/Engine/Source/Programs/UnrealBuildTool/Platform/Linux/UEBuildLinux.cs
@@ -629,6 +629,16 @@ namespace UnrealBuildTool
 			// link with Linux libraries.
 			LinkEnvironment.SystemLibraries.Add("pthread");
 
+			if (Target.bUseAutoRTFMCompiler)
+			{
+				// #jira SOL-9041: Sections prefixes cause relocation errors and thus we avoid using them.
+				string DisableSectionsGuidedSectionPrefix = "-bbsections-guided-section-prefix=0";
+				string DisableProfileGuidedSectionPrefix = "-profile-guided-section-prefix=0";
+
+				CompileEnvironment.AdditionalArguments += $"-Xclang -mllvm -Xclang {DisableSectionsGuidedSectionPrefix} -Xclang -mllvm -Xclang {DisableProfileGuidedSectionPrefix}";
+				LinkEnvironment.AdditionalArguments += $"-Wl,-mllvm -Wl,{DisableSectionsGuidedSectionPrefix} -Wl,-mllvm -Wl,{DisableProfileGuidedSectionPrefix}";
+			}
+
 			// let this class or a sub class do settings specific to that class
 			SetUpSpecificEnvironment(Target, CompileEnvironment, LinkEnvironment);
 		}
diff --git a/Engine/Source/Runtime/AutoRTFM/Private/API.cpp b/Engine/Source/Runtime/AutoRTFM/Private/API.cpp
--- a/Engine/Source/Runtime/AutoRTFM/Private/API.cpp
+++ b/Engine/Source/Runtime/AutoRTFM/Private/API.cpp
@@ -671,37 +671,58 @@ extern "C" autortfm_transaction_id autortfm_current_transaction_id() noexcept
 	return 0;
 }
 
-extern "C" AUTORTFM_INTERNAL UE_AUTORTFM_API void autortfm_abort_transaction_from_closed() noexcept
+// #jira SOL-9048: in theory we should be able to go back to using the .aem file for this.
+AUTORTFM_INTERNAL static void autortfm_abort_transaction_from_closed() noexcept
 {
 	FContext::Get()->AbortTransactionAndThrow(ETransactionStatus::AbortedByRequest);
 }
 
-// Note: Internal.aem maps this to autortfm_abort_transaction_from_closed when called in the closed
 extern "C" void autortfm_abort_transaction() noexcept
 {
-	FContext::Get()->AbortTransaction(ETransactionStatus::AbortedByRequest);
+	if (autortfm_is_closed())
+	{
+		autortfm_abort_transaction_from_closed();
+	}
+	else
+	{
+		FContext::Get()->AbortTransaction(ETransactionStatus::AbortedByRequest);
+	}
 }
 
-extern "C" AUTORTFM_INTERNAL UE_AUTORTFM_API void autortfm_cascading_abort_transaction_from_closed() noexcept
+// #jira SOL-9048: in theory we should be able to go back to using the .aem file for this.
+AUTORTFM_INTERNAL static void autortfm_cascading_abort_transaction_from_closed() noexcept
 {
 	FContext::Get()->AbortTransactionAndThrow(ETransactionStatus::AbortedByCascadingAbort);
 }
 
-// Note: Internal.aem maps this to autortfm_cascading_abort_transaction_from_closed when called in the closed
 extern "C" void autortfm_cascading_abort_transaction() noexcept
 {
-	FContext::Get()->AbortTransaction(ETransactionStatus::AbortedByCascadingAbort);
+	if (autortfm_is_closed())
+	{
+		autortfm_cascading_abort_transaction_from_closed();
+	}
+	else
+	{
+		FContext::Get()->AbortTransaction(ETransactionStatus::AbortedByCascadingAbort);
+	}
 }
 
-extern "C" AUTORTFM_INTERNAL UE_AUTORTFM_API void autortfm_cascading_retry_transaction_from_closed() noexcept
+// #jira SOL-9048: in theory we should be able to go back to using the .aem file for this.
+AUTORTFM_INTERNAL static void autortfm_cascading_retry_transaction_from_closed() noexcept
 {
 	FContext::Get()->AbortTransactionAndThrow(ETransactionStatus::AbortedByCascadingRetry);
 }
 
-// Note: Internal.aem maps this to autortfm_cascading_retry_transaction_from_closed when called in the closed
 extern "C" void autortfm_cascading_retry_transaction() noexcept
 {
-	FContext::Get()->AbortTransaction(ETransactionStatus::AbortedByCascadingRetry);
+	if (autortfm_is_closed())
+	{
+		autortfm_cascading_retry_transaction_from_closed();
+	}
+	else
+	{
+		FContext::Get()->AbortTransaction(ETransactionStatus::AbortedByCascadingRetry);
+	}
 }
 
 extern "C" void autortfm_start_transaction() noexcept
@@ -714,7 +735,8 @@ extern "C" void autortfm_commit_transaction() noexcept
 	FContext::Get()->CommitTransaction();
 }
 
-extern "C" AUTORTFM_INTERNAL UE_AUTORTFM_API void autortfm_open_from_closed(void (*Work)(void*), void* Arg)
+// #jira SOL-9048: in theory we should be able to go back to using the .aem file for this.
+AUTORTFM_INTERNAL static void autortfm_open_from_closed(void (*Work)(void*), void* Arg)
 {
 	FContext* const Context = FContext::Get();
 
@@ -723,13 +745,20 @@ extern "C" AUTORTFM_INTERNAL UE_AUTORTFM_API void autortfm_open_from_closed(void
 	Context->MaybeThrow();
 }
 
-// Note: Internal.aem maps this to autortfm_open_from_closed when called in the closed
 extern "C" void autortfm_open(void (*Work)(void*), void* Arg)
 {
-	Work(Arg);
+	if (autortfm_is_closed())
+	{
+		autortfm_open_from_closed(Work, Arg);
+	}
+	else
+	{
+		Work(Arg);
+	}
 }
 
-extern "C" AUTORTFM_INTERNAL UE_AUTORTFM_API void autortfm_open_from_closed_no_sanitize(void (*Work)(void*), void* Arg)
+// #jira SOL-9048: in theory we should be able to go back to using the .aem file for this.
+AUTORTFM_INTERNAL static void autortfm_open_from_closed_no_sanitize(void (*Work)(void*), void* Arg)
 {
 	FContext* const Context = FContext::Get();
 
@@ -746,10 +775,16 @@ extern "C" AUTORTFM_INTERNAL UE_AUTORTFM_API void autortfm_open_from_closed_no_s
 	AUTORTFM_SANITIZER_ONLY(Sanitizer::EndDisable());
 }
 
-// Note: Internal.aem maps this to autortfm_open_from_closed_no_sanitize when called in the closed
 extern "C" void autortfm_open_no_sanitize(void (*Work)(void*), void* Arg)
 {
-	Work(Arg);
+	if (autortfm_is_closed())
+	{
+		autortfm_open_from_closed_no_sanitize(Work, Arg);
+	}
+	else
+	{
+		Work(Arg);
+	}
 }
 
 extern "C" autortfm_transaction_status autortfm_close(void (*UninstrumentedWork)(void*), void (*InstrumentedWork)(void*), void* Arg)
