Markup functions with the correct AutoRTFM attributes.

#Synced-CL 50768904

diff --git a/Engine/Source/Runtime/AutoRTFM/Private/CompilerInterface.cpp b/Engine/Source/Runtime/AutoRTFM/Private/CompilerInterface.cpp
--- a/Engine/Source/Runtime/AutoRTFM/Private/CompilerInterface.cpp
+++ b/Engine/Source/Runtime/AutoRTFM/Private/CompilerInterface.cpp
@@ -106,14 +106,15 @@ static bool AutoRTFMIsPointerNegative(const void* InPointer)
 	return Address < 0;
 }
 
-extern "C" UE_AUTORTFM_API AUTORTFM_NO_ASAN void autortfm_redirected_load(
+extern "C" AUTORTFM_DISABLE_NO_SANITIZE UE_AUTORTFM_API AUTORTFM_NO_ASAN void autortfm_redirected_load(
 	uint32 AddressSpace, void* DestPointer, uint64 Size, uint64 SourceAddress) noexcept
 {
 	ForTheRuntime::RedirectedLoad(
 		AddressSpace ? AddressSpace : AutoRTFMExtractAddressSpaceFromAddress(SourceAddress), DestPointer, Size, SourceAddress);
 }
 
-extern "C" UE_AUTORTFM_API AUTORTFM_NO_ASAN uint64 autortfm_redirected_load_8(uint32 AddressSpace, uint64 SourceAddress) noexcept
+extern "C" AUTORTFM_DISABLE_NO_SANITIZE UE_AUTORTFM_API AUTORTFM_NO_ASAN uint64 autortfm_redirected_load_8(
+	uint32 AddressSpace, uint64 SourceAddress) noexcept
 {
 	uint64 Result = 0;
 	ForTheRuntime::RedirectedLoad8(
@@ -121,7 +122,8 @@ extern "C" UE_AUTORTFM_API AUTORTFM_NO_ASAN uint64 autortfm_redirected_load_8(ui
 	return Result;
 }
 
-extern "C" UE_AUTORTFM_API AUTORTFM_NO_ASAN uint32 autortfm_redirected_load_4(uint32 AddressSpace, uint64 SourceAddress) noexcept
+extern "C" AUTORTFM_DISABLE_NO_SANITIZE UE_AUTORTFM_API AUTORTFM_NO_ASAN uint32 autortfm_redirected_load_4(
+	uint32 AddressSpace, uint64 SourceAddress) noexcept
 {
 	uint32 Result = 0;
 	ForTheRuntime::RedirectedLoad4(
@@ -129,7 +131,8 @@ extern "C" UE_AUTORTFM_API AUTORTFM_NO_ASAN uint32 autortfm_redirected_load_4(ui
 	return Result;
 }
 
-extern "C" UE_AUTORTFM_API AUTORTFM_NO_ASAN uint16 autortfm_redirected_load_2(uint32 AddressSpace, uint64 SourceAddress) noexcept
+extern "C" AUTORTFM_DISABLE_NO_SANITIZE UE_AUTORTFM_API AUTORTFM_NO_ASAN uint16 autortfm_redirected_load_2(
+	uint32 AddressSpace, uint64 SourceAddress) noexcept
 {
 	uint16 Result = 0;
 	ForTheRuntime::RedirectedLoad2(
@@ -137,7 +140,8 @@ extern "C" UE_AUTORTFM_API AUTORTFM_NO_ASAN uint16 autortfm_redirected_load_2(ui
 	return Result;
 }
 
-extern "C" UE_AUTORTFM_API AUTORTFM_NO_ASAN uint8 autortfm_redirected_load_1(uint32 AddressSpace, uint64 SourceAddress) noexcept
+extern "C" AUTORTFM_DISABLE_NO_SANITIZE UE_AUTORTFM_API AUTORTFM_NO_ASAN uint8 autortfm_redirected_load_1(
+	uint32 AddressSpace, uint64 SourceAddress) noexcept
 {
 	uint8 Result = 0;
 	ForTheRuntime::RedirectedLoad1(
@@ -145,35 +149,35 @@ extern "C" UE_AUTORTFM_API AUTORTFM_NO_ASAN uint8 autortfm_redirected_load_1(uin
 	return Result;
 }
 
-extern "C" UE_AUTORTFM_API AUTORTFM_NO_ASAN void autortfm_redirected_store(
+extern "C" AUTORTFM_DISABLE_NO_SANITIZE UE_AUTORTFM_API AUTORTFM_NO_ASAN void autortfm_redirected_store(
 	uint32 AddressSpace, uint64 DestAddress, uint64 Size, const void* SourcePointer) noexcept
 {
 	ForTheRuntime::RedirectedStore(
 		AddressSpace ? AddressSpace : AutoRTFMExtractAddressSpaceFromAddress(DestAddress), DestAddress, Size, SourcePointer);
 }
 
-extern "C" UE_AUTORTFM_API AUTORTFM_NO_ASAN void autortfm_redirected_store_8(
+extern "C" AUTORTFM_DISABLE_NO_SANITIZE UE_AUTORTFM_API AUTORTFM_NO_ASAN void autortfm_redirected_store_8(
 	uint32 AddressSpace, uint64 DestAddress, uint64 Value) noexcept
 {
 	ForTheRuntime::RedirectedStore8(
 		AddressSpace ? AddressSpace : AutoRTFMExtractAddressSpaceFromAddress(DestAddress), DestAddress, &Value);
 }
 
-extern "C" UE_AUTORTFM_API AUTORTFM_NO_ASAN void autortfm_redirected_store_4(
+extern "C" AUTORTFM_DISABLE_NO_SANITIZE UE_AUTORTFM_API AUTORTFM_NO_ASAN void autortfm_redirected_store_4(
 	uint32 AddressSpace, uint64 DestAddress, uint32 Value) noexcept
 {
 	ForTheRuntime::RedirectedStore4(
 		AddressSpace ? AddressSpace : AutoRTFMExtractAddressSpaceFromAddress(DestAddress), DestAddress, &Value);
 }
 
-extern "C" UE_AUTORTFM_API AUTORTFM_NO_ASAN void autortfm_redirected_store_2(
+extern "C" AUTORTFM_DISABLE_NO_SANITIZE UE_AUTORTFM_API AUTORTFM_NO_ASAN void autortfm_redirected_store_2(
 	uint32 AddressSpace, uint64 DestAddress, uint16 Value) noexcept
 {
 	ForTheRuntime::RedirectedStore2(
 		AddressSpace ? AddressSpace : AutoRTFMExtractAddressSpaceFromAddress(DestAddress), DestAddress, &Value);
 }
 
-extern "C" UE_AUTORTFM_API AUTORTFM_NO_ASAN void autortfm_redirected_store_1(
+extern "C" AUTORTFM_DISABLE_NO_SANITIZE UE_AUTORTFM_API AUTORTFM_NO_ASAN void autortfm_redirected_store_1(
 	uint32 AddressSpace, uint64 DestAddress, uint8 Value) noexcept
 {
 	ForTheRuntime::RedirectedStore1(
@@ -394,7 +398,7 @@ extern "C" AUTORTFM_DISABLE_NO_SANITIZE UE_AUTORTFM_API void autortfm_memset_red
 	}
 }
 
-extern "C" UE_AUTORTFM_ALWAYS_OPEN UE_AUTORTFM_API void autortfm_unreachable(const char* Message) noexcept
+extern "C" void autortfm_unreachable(const char* Message) noexcept
 {
 	AUTORTFM_REPORT_ERROR("AutoRTFM Unreachable: %s", Message);
 	__builtin_unreachable();
diff --git a/Engine/Source/Runtime/AutoRTFM/Public/AutoRTFM/CAPI.h b/Engine/Source/Runtime/AutoRTFM/Public/AutoRTFM/CAPI.h
--- a/Engine/Source/Runtime/AutoRTFM/Public/AutoRTFM/CAPI.h
+++ b/Engine/Source/Runtime/AutoRTFM/Public/AutoRTFM/CAPI.h
@@ -493,7 +493,7 @@ UE_AUTORTFM_CRITICAL_INLINE void autortfm_check_abi(void* ptr, size_t size) AUTO
 	// Called when execution unexpectedly reaches a code path that was considered unreachable.
 	// Either aborts execution of the program or aborts the current transaction, depending on the
 	// current 'InternalAbortAction' state.
-	[[noreturn]] UE_AUTORTFM_API void autortfm_unreachable(const char* Message) AUTORTFM_NOEXCEPT;
+	[[noreturn]] AUTORTFM_OPEN UE_AUTORTFM_API void autortfm_unreachable(const char* Message) AUTORTFM_NOEXCEPT;
 #else
 UE_AUTORTFM_CRITICAL_INLINE void autortfm_unreachable(const char* Message) AUTORTFM_NOEXCEPT
 {
