[AutoRTFM] Fix some incorrect comments

#Synced-CL 50931828

diff --git a/Engine/Source/Runtime/AutoRTFM/Public/AutoRTFM/Constants.h b/Engine/Source/Runtime/AutoRTFM/Public/AutoRTFM/Constants.h
--- a/Engine/Source/Runtime/AutoRTFM/Public/AutoRTFM/Constants.h
+++ b/Engine/Source/Runtime/AutoRTFM/Public/AutoRTFM/Constants.h
@@ -68,12 +68,12 @@ extern "C"
 		autortfm_mode_infer = 3,
 
 		// [[clang::autortfm(open)]] / [[clang::autortfm_always_open]]
-		// Generate AutoRTFM instrumentation for the function.
+		// Do not generate AutoRTFM instrumentation for the function.
 		// Calling the function from the closed will call the uninstrumented variant.
 		autortfm_mode_open = 4,
 
 		// [[clang::autortfm(open_no_sanitize)]] / [[clang::autortfm_always_open_no_sanitize]]
-		// Open, but the AutoRTFM sanitizer will not instrument the function.
+		// Same as autortfm_mode_open, but the AutoRTFM sanitizer will not instrument the function.
 		autortfm_mode_open_no_sanitize = 5,
 
 		// [[clang::autortfm(autortfm_mode_internal)]]
