Check for 0-sized stdlib memory operations as they fall foul of UBSan.

#Synced-CL 51089892

diff --git a/Engine/Source/Runtime/AutoRTFM/Private/Memcpy.cpp b/Engine/Source/Runtime/AutoRTFM/Private/Memcpy.cpp
--- a/Engine/Source/Runtime/AutoRTFM/Private/Memcpy.cpp
+++ b/Engine/Source/Runtime/AutoRTFM/Private/Memcpy.cpp
@@ -13,12 +13,30 @@ namespace AutoRTFM
 void* MemcpyToNew(void* InDst, const void* InSrc, size_t Size, FContext* Context)
 {
 	AUTORTFM_VERBOSE("MemcpyToNew(%p, %p, %zu)", InDst, InSrc, Size);
+
+	// This will eventually go away when `n3322` is adopted, but for now we workaround
+	// this by checking for `Size == 0` and bailing early (the only defined behavior
+	// way going forward).
+	if (AUTORTFM_UNLIKELY(Size == 0))
+	{
+		return nullptr;
+	}
+
 	return memcpy(InDst, InSrc, Size);
 }
 
 void* Memcpy(void* InDst, const void* InSrc, size_t Size, FContext* Context)
 {
 	AUTORTFM_VERBOSE("Memcpy(%p, %p, %zu)", InDst, InSrc, Size);
+
+	// This will eventually go away when `n3322` is adopted, but for now we workaround
+	// this by checking for `Size == 0` and bailing early (the only defined behavior
+	// way going forward).
+	if (AUTORTFM_UNLIKELY(Size == 0))
+	{
+		return nullptr;
+	}
+
 	Context->RecordWrite(InDst, Size);
 	return memcpy(InDst, InSrc, Size);
 }
@@ -26,6 +44,15 @@ void* Memcpy(void* InDst, const void* InSrc, size_t Size, FContext* Context)
 void* Memmove(void* InDst, const void* InSrc, size_t Size, FContext* Context)
 {
 	AUTORTFM_VERBOSE("Memmove(%p, %p, %zu)", InDst, InSrc, Size);
+
+	// This will eventually go away when `n3322` is adopted, but for now we workaround
+	// this by checking for `Size == 0` and bailing early (the only defined behavior
+	// way going forward).
+	if (AUTORTFM_UNLIKELY(Size == 0))
+	{
+		return nullptr;
+	}
+
 	Context->RecordWrite(InDst, Size);
 	return memmove(InDst, InSrc, Size);
 }
@@ -33,6 +60,15 @@ void* Memmove(void* InDst, const void* InSrc, size_t Size, FContext* Context)
 void* Memset(void* InDst, int Value, size_t Size, FContext* Context)
 {
 	AUTORTFM_VERBOSE("Memset(%p, %d, %zu)", InDst, Value, Size);
+
+	// This will eventually go away when `n3322` is adopted, but for now we workaround
+	// this by checking for `Size == 0` and bailing early (the only defined behavior
+	// way going forward).
+	if (AUTORTFM_UNLIKELY(Size == 0))
+	{
+		return nullptr;
+	}
+
 	Context->RecordWrite(InDst, Size);
 	return memset(InDst, Value, Size);
 }
