[AutoRTFM] AUTORTFM_DISABLE THitSet.

The compiler can error in adaptive optimization unity builds, as most if not all callers are disabled.

#Synced-CL 51010198

diff --git a/Engine/Source/Runtime/AutoRTFM/Private/HitSet.h b/Engine/Source/Runtime/AutoRTFM/Private/HitSet.h
--- a/Engine/Source/Runtime/AutoRTFM/Private/HitSet.h
+++ b/Engine/Source/Runtime/AutoRTFM/Private/HitSet.h
@@ -35,7 +35,7 @@ struct FHitSetEntry
 
 static_assert(sizeof(FHitSetEntry) == sizeof(uintptr_t));
 
-class FHitSet final
+class AUTORTFM_DISABLE FHitSet final
 {
 public:
 	// Maximum size in bytes for a hit-set write record.
