[AutoRTFM] Support calling std::swap() from the closed on Windows

#Synced-CL 51090214

diff --git a/Engine/Source/Programs/AutoRTFMTests/Private/Math.cpp b/Engine/Source/Programs/AutoRTFMTests/Private/Math.cpp
--- a/Engine/Source/Programs/AutoRTFMTests/Private/Math.cpp
+++ b/Engine/Source/Programs/AutoRTFMTests/Private/Math.cpp
@@ -12,7 +12,7 @@ namespace
 
 // Check calls Function outside a transaction, in an aborted transaction and
 // committed transaction. The value returned when called in the committed
-// transation is expected to match the value returned when called outside the
+// transaction is expected to match the value returned when called outside the
 // transaction.
 // Function must be the deterministic function with the signature 'T()'.
 template<typename FUNC>
diff --git a/Engine/Source/Programs/AutoRTFMTests/Private/STL.cpp b/Engine/Source/Programs/AutoRTFMTests/Private/STL.cpp
new file mode 100644
--- /dev/null
+++ b/Engine/Source/Programs/AutoRTFMTests/Private/STL.cpp
@@ -0,0 +1,43 @@
+// Copyright Epic Games, Inc. All Rights Reserved.
+
+#include "AutoRTFM.h"
+#include "AutoRTFM/Testing.h"
+#include "Catch2Includes.h"
+
+#include <math.h>
+#include <array>
+#include <utility>
+
+TEST_CASE("STL.swap")
+{
+	auto TestSwap = []<typename T>(T A, T B)
+    {
+        T const OldA = A;
+        T const OldB = B;
+
+        AutoRTFM::Testing::Abort([&]
+        {
+            std::swap(A, B);
+            REQUIRE(A == OldB);
+            REQUIRE(B == OldA);
+            AutoRTFM::AbortTransaction();
+        });
+        REQUIRE(A == OldA);
+        REQUIRE(B == OldB);
+
+        AutoRTFM::Testing::Commit([&]
+        {
+            std::swap(A, B);
+            REQUIRE(A == OldB);
+            REQUIRE(B == OldA);
+        });
+        REQUIRE(A == OldB);
+        REQUIRE(B == OldA);
+    };
+    TestSwap.operator()<char>(1, 2);
+    TestSwap.operator()<int>(1, 2);
+    TestSwap.operator()<float>(1, 2);
+    TestSwap.operator()<const char*>("A", "B");
+    TestSwap.operator()<std::string>("A", "B");
+    TestSwap.operator()<std::array<int, 3>>({1,2,3}, {4,5,6});
+}
diff --git a/Engine/Source/Runtime/AutoRTFM/Private/StdLibSupport.cpp b/Engine/Source/Runtime/AutoRTFM/Private/StdLibSupport.cpp
--- a/Engine/Source/Runtime/AutoRTFM/Private/StdLibSupport.cpp
+++ b/Engine/Source/Runtime/AutoRTFM/Private/StdLibSupport.cpp
@@ -10,6 +10,7 @@
 
 #include <algorithm>
 #include <charconv>
+#include <cstddef>
 #include <float.h>
 #include <functional>  // note: introduces additional math overloads
 #include <stdarg.h>
@@ -235,6 +236,26 @@ extern "C" AUTORTFM_INTERNAL UE_AUTORTFM_API FILE* RTFM___acrt_iob_func(unsigned
 	}
 }
 
+extern "C" AUTORTFM_INTERNAL UE_AUTORTFM_API void RTFM___std_swap_ranges_trivially_swappable_noalias(
+	void* FirstA, void* LastA, void* FirstB) noexcept
+{
+	if (AUTORTFM_LIKELY(LastA > FirstA))
+	{
+		std::byte* PtrA = reinterpret_cast<std::byte*>(FirstA);
+		std::byte* PtrB = reinterpret_cast<std::byte*>(FirstB);
+		size_t const Size = reinterpret_cast<std::byte*>(LastA) - PtrA;
+		FContext* Context = FContext::Get();
+		Context->RecordWrite(PtrA, Size);
+		Context->RecordWrite(PtrB, Size);
+		for (size_t I = 0; I < Size; I++, PtrA++, PtrB++)
+		{
+			std::byte Tmp = *PtrA;
+			*PtrA = *PtrB;
+			*PtrB = Tmp;
+		}
+	}
+}
+
 // FIXME: Does not currently support %n format specifiers.
 extern "C" AUTORTFM_INTERNAL UE_AUTORTFM_API int RTFM___stdio_common_vfprintf(
 	unsigned __int64 Options, FILE* Stream, char const* Format, _locale_t Locale, va_list ArgList)
