[AutoRTFM] Hook up new / delete external mappings for Windows

Allows clang sanitizers (ASAN, etc) to be used on Windows with AutoRTFM enabled.

#Synced-CL 51010198

diff --git a/Engine/Source/Runtime/AutoRTFM/AutoRTFM.Build.cs b/Engine/Source/Runtime/AutoRTFM/AutoRTFM.Build.cs
--- a/Engine/Source/Runtime/AutoRTFM/AutoRTFM.Build.cs
+++ b/Engine/Source/Runtime/AutoRTFM/AutoRTFM.Build.cs
@@ -32,6 +32,11 @@ namespace UnrealBuildTool.Rules
 				{
 					AutoRTFMExternalMappingFiles.Add("Runtime/AutoRTFM/Public/UBSAN.aem");
 				}
+				
+				if (Target.StaticAllocator == StaticAllocatorType.Ansi)
+				{
+					AutoRTFMExternalMappingFiles.Add("Runtime/AutoRTFM/Public/NewDelete.Windows.aem");
+				}
 			}
 			if (Target.Platform == UnrealTargetPlatform.Linux || Target.Platform == UnrealTargetPlatform.LinuxArm64)
 			{
diff --git a/Engine/Source/Runtime/AutoRTFM/Private/StdLibSupport.cpp b/Engine/Source/Runtime/AutoRTFM/Private/StdLibSupport.cpp
--- a/Engine/Source/Runtime/AutoRTFM/Private/StdLibSupport.cpp
+++ b/Engine/Source/Runtime/AutoRTFM/Private/StdLibSupport.cpp
@@ -833,7 +833,7 @@ extern "C" void __cxa_end_catch();
 #endif
 
 template <typename OpenNewFn, typename OpenDeleteFn, typename... ArgTys>
-UE_AUTORTFM_FORCEINLINE static void* RTFM_ClosedNew(OpenNewFn* New, OpenDeleteFn* Delete, unsigned long Size, ArgTys... Args)
+UE_AUTORTFM_FORCEINLINE static void* RTFM_ClosedNew(OpenNewFn* New, OpenDeleteFn* Delete, size_t Size, ArgTys... Args)
 {
 	void* Result = New(Size, Args...);
 	FContext* Context = FContext::Get();
@@ -852,147 +852,179 @@ UE_AUTORTFM_FORCEINLINE static void RTFM_ClosedDelete(OpenDeleteFn* Delete, void
 	}
 }
 
-#if AUTORTFM_PLATFORM_LINUX
-// vvvvvvvvvvvvvvvvvvvvvvvvvvvvvv AUTORTFM_PLATFORM_LINUX vvvvvvvvvvvvvvvvvvvvvvvvvvvvvv
-// LibCxx's string.h uses builtin string functions inside inline functions.
-// We can't take the address of these builtins, but we know they map to these C functions.
-extern "C"
+// operator new(size_t)
+// ??2@YAPEAX_K@Z
+extern "C" AUTORTFM_INTERNAL UE_AUTORTFM_API void* RTFM_Znwm(size_t Size)
 {
-	char* strchr(const char*, int);
-	char* strrchr(const char*, int);
-	char* strstr(const char*, const char*);
-
-	// Operator new and new[]
-	void* _Znwm(unsigned long);
-	void* _Znam(unsigned long);
-	void* _ZnwmRKSt9nothrow_t(unsigned long, const std::nothrow_t&);
-	void* _ZnamRKSt9nothrow_t(unsigned long, const std::nothrow_t&);
-	void* _ZnwmSt11align_val_t(unsigned long, std::align_val_t);
-	void* _ZnamSt11align_val_t(unsigned long, std::align_val_t);
-	void* _ZnwmSt11align_val_tRKSt9nothrow_t(unsigned long, std::align_val_t, const std::nothrow_t&);
-	void* _ZnamSt11align_val_tRKSt9nothrow_t(unsigned long, std::align_val_t, const std::nothrow_t&);
-	// Operator delete and delete[]
-	void _ZdlPv(void*);
-	void _ZdaPv(void*);
-	void _ZdlPvRKSt9nothrow_t(void*, const std::nothrow_t&);
-	void _ZdaPvRKSt9nothrow_t(void*, const std::nothrow_t&);
-	void _ZdlPvm(void*, unsigned long);
-	void _ZdaPvm(void*, unsigned long);
-	void _ZdlPvSt11align_val_t(void*, std::align_val_t);
-	void _ZdaPvSt11align_val_t(void*, std::align_val_t);
-	void _ZdlPvSt11align_val_tRKSt9nothrow_t(void*, std::align_val_t, const std::nothrow_t&);
-	void _ZdaPvSt11align_val_tRKSt9nothrow_t(void*, std::align_val_t, const std::nothrow_t&);
-	void _ZdlPvmSt11align_val_t(void*, unsigned long, std::align_val_t);
-	void _ZdaPvmSt11align_val_t(void*, unsigned long, std::align_val_t);
-}  // extern "C"
-
-extern "C" AUTORTFM_INTERNAL UE_AUTORTFM_API void* RTFM_Znwm(unsigned long Size)
-{
-	return RTFM_ClosedNew(_Znwm, _ZdlPv, Size);
+	void* (*New)(size_t) = &operator new;
+	void (*Delete)(void*) = &operator delete;
+	return RTFM_ClosedNew(New, Delete, Size);
 }
 
-extern "C" AUTORTFM_INTERNAL UE_AUTORTFM_API void* RTFM_Znam(unsigned long Size)
+// operator new[](size_t)
+// ??_U@YAPEAX_K@Z
+extern "C" AUTORTFM_INTERNAL UE_AUTORTFM_API void* RTFM_Znam(size_t Size)
 {
-	return RTFM_ClosedNew(_Znam, _ZdaPv, Size);
+	void* (*New)(size_t) = &operator new[];
+	void (*Delete)(void*) = &operator delete[];
+	return RTFM_ClosedNew(New, Delete, Size);
 }
 
-extern "C" AUTORTFM_INTERNAL UE_AUTORTFM_API void* RTFM_ZnwmRKSt9nothrow_t(unsigned long Size, const std::nothrow_t& NoThrow)
+// operator new(size_t, std::nothrow_t const&)
+// ??2@YAPEAX_KAEBUnothrow_t@std@@@Z
+extern "C" AUTORTFM_INTERNAL UE_AUTORTFM_API void* RTFM_ZnwmRKSt9nothrow_t(size_t Size, std::nothrow_t const& NoThrow)
 {
-	return RTFM_ClosedNew(_ZnwmRKSt9nothrow_t, _ZdlPvRKSt9nothrow_t, Size, NoThrow);
+	void* (*New)(size_t, std::nothrow_t const&) = &operator new;
+	void (*Delete)(void*, std::nothrow_t const&) = &operator delete;
+	return RTFM_ClosedNew(New, Delete, Size, NoThrow);
 }
 
-extern "C" AUTORTFM_INTERNAL UE_AUTORTFM_API void* RTFM_ZnamRKSt9nothrow_t(unsigned long Size, const std::nothrow_t& NoThrow)
+// operator new(size_t, std::nothrow_t const&)
+// ??_U@YAPEAX_KAEBUnothrow_t@std@@@Z
+extern "C" AUTORTFM_INTERNAL UE_AUTORTFM_API void* RTFM_ZnamRKSt9nothrow_t(size_t Size, std::nothrow_t const& NoThrow)
 {
-	return RTFM_ClosedNew(_ZnamRKSt9nothrow_t, _ZdaPvRKSt9nothrow_t, Size, NoThrow);
+	void* (*New)(size_t, std::nothrow_t const& NoThrow) = &operator new[];
+	void (*Delete)(void*, std::nothrow_t const& NoThrow) = &operator delete[];
+	return RTFM_ClosedNew(New, Delete, Size, NoThrow);
 }
 
-extern "C" AUTORTFM_INTERNAL UE_AUTORTFM_API void* RTFM_ZnwmSt11align_val_t(unsigned long Size, std::align_val_t Align)
+// operator new(size_t, std::align_val_t)
+// ??2@YAPEAX_KW4align_val_t@std@@@Z
+extern "C" AUTORTFM_INTERNAL UE_AUTORTFM_API void* RTFM_ZnwmSt11align_val_t(size_t Size, std::align_val_t Align)
 {
-	return RTFM_ClosedNew(_ZnwmSt11align_val_t, _ZdlPvSt11align_val_t, Size, Align);
+	void* (*New)(size_t, std::align_val_t) = &operator new;
+	void (*Delete)(void*, std::align_val_t) = &operator delete;
+	return RTFM_ClosedNew(New, Delete, Size, Align);
 }
 
-extern "C" AUTORTFM_INTERNAL UE_AUTORTFM_API void* RTFM_ZnamSt11align_val_t(unsigned long Size, std::align_val_t Align)
+// operator new[](size_t, std::align_val_t)
+// ??_U@YAPEAX_KW4align_val_t@std@@@Z
+extern "C" AUTORTFM_INTERNAL UE_AUTORTFM_API void* RTFM_ZnamSt11align_val_t(size_t Size, std::align_val_t Align)
 {
-	return RTFM_ClosedNew(_ZnamSt11align_val_t, _ZdaPvSt11align_val_t, Size, Align);
+	void* (*New)(size_t, std::align_val_t) = &operator new[];
+	void (*Delete)(void*, std::align_val_t) = &operator delete[];
+	return RTFM_ClosedNew(New, Delete, Size, Align);
 }
 
+// operator new(size_t Size, std::align_val_t, std::nothrow_t const&)
+// ??2@YAPEAX_KW4align_val_t@std@@AEBUnothrow_t@1@@Z
 extern "C" AUTORTFM_INTERNAL UE_AUTORTFM_API void* RTFM_ZnwmSt11align_val_tRKSt9nothrow_t(
-	unsigned long Size, std::align_val_t Align, const std::nothrow_t& NoThrow)
+	size_t Size, std::align_val_t Align, std::nothrow_t const& NoThrow)
 {
-	return RTFM_ClosedNew(_ZnwmSt11align_val_tRKSt9nothrow_t, _ZdlPvSt11align_val_tRKSt9nothrow_t, Size, Align, NoThrow);
+	void* (*New)(size_t, std::align_val_t, std::nothrow_t const&) = &operator new;
+	void (*Delete)(void*, std::align_val_t, std::nothrow_t const&) = &operator delete;
+	return RTFM_ClosedNew(New, Delete, Size, Align, NoThrow);
 }
 
+// operator new[](size_t Size, std::align_val_t, std::nothrow_t const&)
+// ??_U@YAPEAX_KW4align_val_t@std@@AEBUnothrow_t@1@@Z
 extern "C" AUTORTFM_INTERNAL UE_AUTORTFM_API void* RTFM_ZnamSt11align_val_tRKSt9nothrow_t(
-	unsigned long Size, std::align_val_t Align, const std::nothrow_t& NoThrow)
+	size_t Size, std::align_val_t Align, std::nothrow_t const& NoThrow)
 {
-	return RTFM_ClosedNew(_ZnamSt11align_val_tRKSt9nothrow_t, _ZdaPvSt11align_val_tRKSt9nothrow_t, Size, Align, NoThrow);
+	void* (*New)(size_t, std::align_val_t, std::nothrow_t const&) = &operator new[];
+	void (*Delete)(void*, std::align_val_t, std::nothrow_t const&) = &operator delete[];
+	return RTFM_ClosedNew(New, Delete, Size, Align, NoThrow);
 }
 
+// operator delete(void*)
+// ??3@YAXPEAX@Z
 extern "C" AUTORTFM_INTERNAL UE_AUTORTFM_API void RTFM_ZdlPv(void* Pointer)
 {
-	RTFM_ClosedDelete(_ZdlPv, Pointer);
+	void (*Delete)(void*) = &operator delete;
+	RTFM_ClosedDelete(Delete, Pointer);
 }
 
+// operator delete[](void*)
+// ??_V@YAXPEAX@Z
 extern "C" AUTORTFM_INTERNAL UE_AUTORTFM_API void RTFM_ZdaPv(void* Pointer)
 {
-	RTFM_ClosedDelete(_ZdaPv, Pointer);
+	void (*Delete)(void*) = &operator delete[];
+	RTFM_ClosedDelete(Delete, Pointer);
 }
 
-extern "C" AUTORTFM_INTERNAL UE_AUTORTFM_API void RTFM_ZdlPvRKSt9nothrow_t(void* Pointer, const std::nothrow_t& NoThrow)
+// operator delete(void*, std::nothrow_t const&)
+// ??3@YAXPEAXAEBUnothrow_t@std@@@Z
+extern "C" AUTORTFM_INTERNAL UE_AUTORTFM_API void RTFM_ZdlPvRKSt9nothrow_t(void* Pointer, std::nothrow_t const& NoThrow)
 {
-	RTFM_ClosedDelete(_ZdlPvRKSt9nothrow_t, Pointer, NoThrow);
+	void (*Delete)(void*, std::nothrow_t const&) = &operator delete;
+	RTFM_ClosedDelete(Delete, Pointer, NoThrow);
 }
 
-extern "C" AUTORTFM_INTERNAL UE_AUTORTFM_API void RTFM_ZdaPvRKSt9nothrow_t(void* Pointer, const std::nothrow_t& NoThrow)
+// operator delete[](void *, std::nothrow_t const&)
+// ??_V@YAXPEAXAEBUnothrow_t@std@@@Z
+extern "C" AUTORTFM_INTERNAL UE_AUTORTFM_API void RTFM_ZdaPvRKSt9nothrow_t(void* Pointer, std::nothrow_t const& NoThrow)
 {
-	RTFM_ClosedDelete(_ZdaPvRKSt9nothrow_t, Pointer, NoThrow);
+	void (*Delete)(void*, std::nothrow_t const&) = &operator delete[];
+	RTFM_ClosedDelete(Delete, Pointer, NoThrow);
 }
 
-extern "C" AUTORTFM_INTERNAL UE_AUTORTFM_API void RTFM_ZdlPvm(void* Pointer, unsigned long Size)
+// operator delete(void*, size_t)
+// ??3@YAXPEAX_K@Z
+extern "C" AUTORTFM_INTERNAL UE_AUTORTFM_API void RTFM_ZdlPvm(void* Pointer, size_t Size)
 {
-	RTFM_ClosedDelete(_ZdlPvm, Pointer, Size);
+	void (*Delete)(void*, size_t) = &operator delete;
+	RTFM_ClosedDelete(Delete, Pointer, Size);
 }
 
-extern "C" AUTORTFM_INTERNAL UE_AUTORTFM_API void RTFM_ZdaPvm(void* Pointer, unsigned long Size)
+// operator delete[](void*, size_t)
+// ??_V@YAXPEAX_K@Z
+extern "C" AUTORTFM_INTERNAL UE_AUTORTFM_API void RTFM_ZdaPvm(void* Pointer, size_t Size)
 {
-	RTFM_ClosedDelete(_ZdaPvm, Pointer, Size);
+	void (*Delete)(void*, size_t) = &operator delete[];
+	RTFM_ClosedDelete(Delete, Pointer, Size);
 }
 
+// operator delete(void*, std::align_val_t)
+// ??3@YAXPEAXW4align_val_t@std@@@Z
 extern "C" AUTORTFM_INTERNAL UE_AUTORTFM_API void RTFM_ZdlPvSt11align_val_t(void* Pointer, std::align_val_t Align)
 {
-	RTFM_ClosedDelete(_ZdlPvSt11align_val_t, Pointer, Align);
+	void (*Delete)(void*, std::align_val_t) = &operator delete;
+	RTFM_ClosedDelete(Delete, Pointer, Align);
 }
 
+// operator delete[](void*, std::align_val_t)
+// ??_V@YAXPEAXW4align_val_t@std@@@Z
 extern "C" AUTORTFM_INTERNAL UE_AUTORTFM_API void RTFM_ZdaPvSt11align_val_t(void* Pointer, std::align_val_t Align)
 {
-	RTFM_ClosedDelete(_ZdaPvSt11align_val_t, Pointer, Align);
+	void (*Delete)(void*, std::align_val_t) = &operator delete[];
+	RTFM_ClosedDelete(Delete, Pointer, Align);
 }
 
+// operator delete(void*, std::align_val_t, std::nothrow_t const&)
+// ??3@YAXPEAXW4align_val_t@std@@AEBUnothrow_t@1@@Z
 extern "C" AUTORTFM_INTERNAL UE_AUTORTFM_API void RTFM_ZdlPvSt11align_val_tRKSt9nothrow_t(
-	void* Pointer, std::align_val_t Align, const std::nothrow_t& NoThrow)
+	void* Pointer, std::align_val_t Align, std::nothrow_t const& NoThrow)
 {
-	RTFM_ClosedDelete(_ZdlPvSt11align_val_tRKSt9nothrow_t, Pointer, Align, NoThrow);
+	void (*Delete)(void*, std::align_val_t, std::nothrow_t const&) = &operator delete;
+	RTFM_ClosedDelete(Delete, Pointer, Align, NoThrow);
 }
 
+// operator delete[](void*, std::align_val_t, std::nothrow_t const&)
+// ??_V@YAXPEAXW4align_val_t@std@@AEBUnothrow_t@1@@Z
 extern "C" AUTORTFM_INTERNAL UE_AUTORTFM_API void RTFM_ZdaPvSt11align_val_tRKSt9nothrow_t(
-	void* Pointer, std::align_val_t Align, const std::nothrow_t& NoThrow)
+	void* Pointer, std::align_val_t Align, std::nothrow_t const& NoThrow)
 {
-	RTFM_ClosedDelete(_ZdaPvSt11align_val_tRKSt9nothrow_t, Pointer, Align, NoThrow);
+	void (*Delete)(void*, std::align_val_t, std::nothrow_t const&) = &operator delete[];
+	RTFM_ClosedDelete(Delete, Pointer, Align, NoThrow);
 }
 
+// operator delete(void*, size_t, std::align_val_t)
+// ??3@YAXPEAX_KW4align_val_t@std@@@Z
 extern "C" AUTORTFM_INTERNAL UE_AUTORTFM_API void RTFM_ZdlPvmSt11align_val_t(
 	void* Pointer, unsigned long Size, std::align_val_t Align)
 {
-	RTFM_ClosedDelete(_ZdlPvmSt11align_val_t, Pointer, Size, Align);
+	void (*Delete)(void*, size_t, std::align_val_t) = &operator delete;
+	RTFM_ClosedDelete(Delete, Pointer, Size, Align);
 }
 
+// operator delete[](void*, size_t, std::align_val_t)
+// ??_V@YAXPEAX_KW4align_val_t@std@@@Z
 extern "C" AUTORTFM_INTERNAL UE_AUTORTFM_API void RTFM_ZdaPvmSt11align_val_t(
 	void* Pointer, unsigned long Size, std::align_val_t Align)
 {
-	RTFM_ClosedDelete(_ZdaPvmSt11align_val_t, Pointer, Size, Align);
+	void (*Delete)(void*, size_t, std::align_val_t) = &operator delete[];
+	RTFM_ClosedDelete(Delete, Pointer, Size, Align);
 }
-// ^^^^^^^^^^^^^^^^^^^^^^^^^^^^^^ AUTORTFM_PLATFORM_LINUX ^^^^^^^^^^^^^^^^^^^^^^^^^^^^^^
-#endif
 
 }  // namespace AutoRTFM
 
