[AutoRTFM] Split TIntervalTree into TIntervalTree and FMemoryTracker and add documentation

First step towards fixing FIntervalTree::Contains() for overlaps when start pointer is not inside an inserted interval and for fixing incorrect expands / merges with data - fixes both required by the AutoRTFM sanitizer, not the memory tracker.

#Synced-CL 51461940

diff --git a/Engine/Source/Programs/AutoRTFMTests/Private/IntervalTestUtils.h b/Engine/Source/Programs/AutoRTFMTests/Private/IntervalTestUtils.h
new file mode 100644
--- /dev/null
+++ b/Engine/Source/Programs/AutoRTFMTests/Private/IntervalTestUtils.h
@@ -0,0 +1,84 @@
+// Copyright Epic Games, Inc. All Rights Reserved.
+
+#pragma once
+
+#include "Interval.h"
+#include "Catch2Includes.h"
+
+#include <cstdint>
+#include <limits>
+#include <vector>
+
+using FInterval = AutoRTFM::FInterval;
+
+static constexpr FInterval MaxInterval{static_cast<uintptr_t>(0), std::numeric_limits<uintptr_t>::max()};
+
+struct FTreeIntervals
+{
+	// Intervals that are not in the tree
+	std::vector<FInterval> Free;
+	// Intervals that are in the tree
+	std::vector<FInterval> Used;
+
+	// Sorts the Free and Used intervals into ascending order, and checks that
+	// each interval in the vector has a non-zero size and does not overlap any
+	// other interval in the same vector.
+	void SortAndValidate()
+	{
+		auto SortAndValidateIntervals = [](std::vector<FInterval>& Intervals)
+		{
+			std::sort(Intervals.begin(), Intervals.end());
+			uintptr_t End = 0;
+			for (FInterval& Interval : Intervals)
+			{
+				REQUIRE(Interval.Size() > 0);
+				REQUIRE(End <= Interval.Start);
+				End = Interval.End;
+			}
+		};
+		SortAndValidateIntervals(Free);
+		SortAndValidateIntervals(Used);
+	}
+};
+
+// Populates a tree with at most MaxUsed random intervals across the entire uintptr_t interval range.
+static FTreeIntervals BuildIntervals(int Seed, size_t MaxUsed)
+{
+	Catch::SimplePcg32 Rand(Seed);
+	std::vector<FInterval> Free{MaxInterval};
+	std::vector<FInterval> Used;
+
+	while (!Free.empty() && Used.size() < MaxUsed)
+	{
+		// Pop a random free interval
+		std::swap(Free[Rand() % Free.size()], Free.back());
+		const FInterval FreeInterval = Free.back();
+		Free.pop_back();
+
+		// Steal the part of the interval
+		static constexpr size_t MaxIntervalSize = 0x100000;
+		const uintptr_t Size = FreeInterval.Size() > 1 ? (1 + Rand() % std::min(MaxIntervalSize, FreeInterval.Size() - 1)) : 1;
+		uintptr_t Offset = FreeInterval.Size() > Size ? (Rand() % (FreeInterval.Size() - Size)) : 0;
+		if (Rand() % 100 < 5) // 5% of the intervals are aligned to the start or end of another interval
+		{
+			Offset = (Rand() & 1) ? 0 : FreeInterval.Size() - Size;
+		}
+
+		REQUIRE(Offset + Size <= FreeInterval.Size());
+
+		const FInterval UsedInterval{FreeInterval.Start + Offset, FreeInterval.Start + Offset + Size};
+
+		if (FreeInterval.Start != UsedInterval.Start)
+		{
+			Free.push_back(FInterval{FreeInterval.Start, UsedInterval.Start});
+		}
+		if (UsedInterval.End != FreeInterval.End)
+		{
+			Free.push_back(FInterval{UsedInterval.End, FreeInterval.End});
+		}
+
+		Used.push_back(UsedInterval);
+	}
+
+	return FTreeIntervals{Free, Used};
+}
diff --git a/Engine/Source/Programs/AutoRTFMTests/Private/IntervalTreeTests.cpp b/Engine/Source/Programs/AutoRTFMTests/Private/IntervalTreeTests.cpp
--- a/Engine/Source/Programs/AutoRTFMTests/Private/IntervalTreeTests.cpp
+++ b/Engine/Source/Programs/AutoRTFMTests/Private/IntervalTreeTests.cpp
@@ -1,152 +1,37 @@
 // Copyright Epic Games, Inc. All Rights Reserved.
 
 #include "AutoRTFM.h"
+#include "IntervalTestUtils.h"
 #include "IntervalTree.h"
-#include "Catch2Includes.h"
-#include <cstdint>
-#include <limits>
-#include <random>
-#include <type_traits>
 
-using FInterval = AutoRTFM::FInterval;
+template <typename T> using TIntervalTree = AutoRTFM::TIntervalTree<T>;
 
-static constexpr FInterval MaxInterval{static_cast<uintptr_t>(0), std::numeric_limits<uintptr_t>::max()};
-
-struct FTreeIntervals
+struct FIntervalAndData
 {
-	// Intervals that are not in the tree
-	std::vector<FInterval> Free;
-	// Intervals that are in the tree
-	std::vector<FInterval> Used;
-	
-	// Sorts the Free and Used intervals into ascending order, and checks that
-	// each interval in the vector has a non-zero size and does not overlap any
-	// other interval in the same vector.
-	void SortAndValidate()
-	{
-		auto SortAndValidateIntervals = [](std::vector<FInterval>& Intervals)
-		{
-			std::sort(Intervals.begin(), Intervals.end());
-			uintptr_t End = 0;
-			for (FInterval& Interval : Intervals)
-			{
-				REQUIRE(Interval.Size() > 0);
-				REQUIRE(End <= Interval.Start);
-				End = Interval.End;
-			}
-		};
-		SortAndValidateIntervals(Free);
-		SortAndValidateIntervals(Used);
-	}
+	FInterval Interval;
+	int Data;
 };
 
-// Populates Tree with at most MaxUsed random intervals across the entire uintptr_t interval range.
-static FTreeIntervals BuildIntervals(int Seed, size_t MaxUsed)
-{
-	Catch::SimplePcg32 Rand(Seed);
-	std::vector<FInterval> Free{MaxInterval};
-	std::vector<FInterval> Used;
-
-	while (!Free.empty() && Used.size() < MaxUsed)
-	{
-		// Pop a random free interval
-		std::swap(Free[Rand() % Free.size()], Free.back());
-		const FInterval FreeInterval = Free.back();
-		Free.pop_back();
-
-		// Steal the part of the interval
-		static constexpr size_t MaxIntervalSize = 0x100000;
-		const uintptr_t Size = FreeInterval.Size() > 1 ? (1 + Rand() % std::min(MaxIntervalSize, FreeInterval.Size() - 1)) : 1;
-		uintptr_t Offset = FreeInterval.Size() > Size ? (Rand() % (FreeInterval.Size() - Size)) : 0;
-		if (Rand() % 100 < 5) // 5% of the intervals are aligned to the start or end of another interval
-		{
-			Offset = (Rand() & 1) ? 0 : FreeInterval.Size() - Size;
-		}
-		
-		REQUIRE(Offset + Size <= FreeInterval.Size());
-
-		const FInterval UsedInterval{FreeInterval.Start + Offset, FreeInterval.Start + Offset + Size};
-		
-		if (FreeInterval.Start != UsedInterval.Start)
-		{
-			Free.push_back(FInterval{FreeInterval.Start, UsedInterval.Start});
-		}
-		if (UsedInterval.End != FreeInterval.End)
-		{
-			Free.push_back(FInterval{UsedInterval.End, FreeInterval.End});
-		}
-
-		Used.push_back(UsedInterval);
-	}
-
-	return FTreeIntervals{Free, Used};
-}
-
-TEMPLATE_TEST_CASE("IntervalTree", "", void, int)
+TEST_CASE("IntervalTree")
 {
 	AutoRTFM::UnreachableIfClosed();
 
-	using IntervalTreeType = AutoRTFM::TIntervalTree<TestType>;
-	struct FIntervalAndData
-	{
-		FInterval Interval;
-		int Data;
-	};
-	static constexpr bool bHasData = !std::is_same_v<TestType, void>;
-
 	// Convenience lambda to convert a uintptr_t to a void*
 	auto Ptr = [](uintptr_t Address) { return reinterpret_cast<void*>(Address); };
 
-	using EntryType = std::conditional_t<bHasData, typename IntervalTreeType::FIntervalAndData, FInterval>;
-	auto IntervalOf = [] AUTORTFM_DISABLE (EntryType Entry)
+	auto Insert = [Ptr](TIntervalTree<int>& Tree, uintptr_t Address, uintptr_t Size, int Data)
 	{
-		if constexpr (bHasData)
-		{
-			return Entry.Interval;
-		}
-		else
-		{
-			return Entry;
-		}
-	};
-	auto MakeEntry = [] AUTORTFM_DISABLE (uintptr_t Address, uintptr_t Size, int Data)
-	{
-		if constexpr (bHasData)
-		{
-			return EntryType{FInterval{Address, Size}, Data};
-		}
-		else
-		{
-			return FInterval{Address, Size};
-		}
-	};
-	auto Insert = [Ptr] AUTORTFM_DISABLE (IntervalTreeType& Tree, uintptr_t Address, uintptr_t Size, int Data)
-	{
-		if constexpr (bHasData)
-		{
-			Tree.Insert(Ptr(Address), Size, Data);
-		}
-		else
-		{
-			Tree.Insert(Ptr(Address), Size);
-		}
+		Tree.Insert(Ptr(Address), Size, Data);
 	};
-	
-	// Checks that the interval tree has all of the expected intervals, and that Contains() returns
+
+	// Checks that the interval tree has all of the expected intervals, and that Get() returns
 	// the correct value for various intervals.
-	auto Validate = [&] AUTORTFM_DISABLE (IntervalTreeType& IntervalTree, std::vector<FIntervalAndData> Expected)
+	auto Validate = [&] AUTORTFM_DISABLE (TIntervalTree<int>& IntervalTree, std::vector<FIntervalAndData> Expected)
 	{
 		std::vector<FIntervalAndData> Got;
 		for (auto Entry : IntervalTree)
 		{
-			if constexpr (bHasData)
-			{
-				Got.push_back(FIntervalAndData{Entry.Interval, Entry.Data});
-			}
-			else
-			{
-				Got.push_back(FIntervalAndData{Entry, {}});
-			}
+			Got.push_back(FIntervalAndData{Entry.Interval, Entry.Data});
 		}
 
 		REQUIRE(Expected.empty() == IntervalTree.IsEmpty());
@@ -154,10 +39,7 @@ TEMPLATE_TEST_CASE("IntervalTree", "", void, int)
 		{
 			REQUIRE(Expected[I].Interval.Start == Got[I].Interval.Start);
 			REQUIRE(Expected[I].Interval.End == Got[I].Interval.End);
-			if constexpr (bHasData)
-			{
-				REQUIRE(Expected[I].Data == Got[I].Data);
-			}
+			REQUIRE(Expected[I].Data == Got[I].Data);
 		}
 		REQUIRE(Expected.size() == Got.size());
 		for (size_t Index = 0; Index < Expected.size(); Index++)
@@ -165,14 +47,14 @@ TEMPLATE_TEST_CASE("IntervalTree", "", void, int)
 			const FIntervalAndData Entry = Expected[Index];
 			const bool bTouchingPrev = Index > 0 && Expected[Index - 1].Interval.End == Entry.Interval.Start;
 			const bool bTouchingNext = Index < (Expected.size() - 1) && Expected[Index + 1].Interval.Start == Entry.Interval.End;
-			REQUIRE(true == IntervalTree.Contains(Ptr(Entry.Interval.Start), 1));
-			REQUIRE(true == IntervalTree.Contains(Ptr(Entry.Interval.Start), Entry.Interval.Size()));
-			REQUIRE(bTouchingPrev == IntervalTree.Contains(Ptr(Entry.Interval.Start - 1), 1));
-			REQUIRE(bTouchingNext == IntervalTree.Contains(Ptr(Entry.Interval.End), 1));
+			REQUIRE(nullptr != IntervalTree.Get(Ptr(Entry.Interval.Start), 1));
+			REQUIRE(nullptr != IntervalTree.Get(Ptr(Entry.Interval.Start), Entry.Interval.Size()));
+			REQUIRE(bTouchingPrev == (nullptr != IntervalTree.Get(Ptr(Entry.Interval.Start - 1), 1)));
+			REQUIRE(bTouchingNext == (nullptr != IntervalTree.Get(Ptr(Entry.Interval.End), 1)));
 		}
 	};
 
-	IntervalTreeType IntervalTree;
+	TIntervalTree<int> IntervalTree;
 
 	Validate(IntervalTree, {});
 
@@ -202,10 +84,10 @@ TEMPLATE_TEST_CASE("IntervalTree", "", void, int)
 
 		Insert(IntervalTree, 40, 10, 6);
 		Validate(IntervalTree, {{{20, 30}, 2}, {{40, 50}, 1}, {{60, 70}, 3}, {{80, 90}, 4}});
-		
+
 		Insert(IntervalTree, 60, 10, 7);
 		Validate(IntervalTree, {{{20, 30}, 2}, {{40, 50}, 1}, {{60, 70}, 3}, {{80, 90}, 4}});
-		
+
 		Insert(IntervalTree, 80, 10, 8);
 		Validate(IntervalTree, {{{20, 30}, 2}, {{40, 50}, 1}, {{60, 70}, 3}, {{80, 90}, 4}});
 
@@ -249,21 +131,21 @@ TEMPLATE_TEST_CASE("IntervalTree", "", void, int)
 
 		SECTION("Intertwined")
 		{
-			IntervalTreeType IntervalTreeB;
+			TIntervalTree<int> IntervalTreeB;
 			Insert(IntervalTreeB, 10, 5, 10);
 			Insert(IntervalTreeB, 33, 6, 11);
 			Insert(IntervalTreeB, 54, 2, 12);
 			Insert(IntervalTreeB, 73, 3, 13);
 			Insert(IntervalTreeB, 95, 2, 14);
 			Validate(IntervalTreeB, {{{10, 15}, 10}, {{33, 39}, 11}, {{54, 56}, 12}, {{73, 76}, 13}, {{95, 97}, 14}});
-			
+
 			IntervalTree.Merge(IntervalTreeB);
 			Validate(IntervalTree, {{{10, 15}, 10}, {{20, 30}, 1}, {{33, 39}, 11}, {{40, 50}, 2}, {{54, 56}, 12}, {{60, 70}, 3}, {{73, 76}, 13}, {{80, 90}, 4}, {{95, 97}, 14}});
 		}
 
 		SECTION("Expand")
 		{
-			IntervalTreeType IntervalTreeB;
+			TIntervalTree<int> IntervalTreeB;
 			Insert(IntervalTreeB, 15, 5, 10);
 			Insert(IntervalTreeB, 30, 3, 11);
 			Insert(IntervalTreeB, 37, 3, 12);
@@ -273,7 +155,7 @@ TEMPLATE_TEST_CASE("IntervalTree", "", void, int)
 			Insert(IntervalTreeB, 78, 2, 16);
 			Insert(IntervalTreeB, 90, 3, 17);
 			Validate(IntervalTreeB, {{{15, 20}, 10}, {{30, 33}, 11}, {{37, 40}, 12}, {{50, 52}, 13}, {{55, 60}, 14}, {{70, 74}, 15}, {{78, 80}, 16}, {{90, 93}, 17}});
-			
+
 			// TODO(SOL-8824): This is incorrect for intervals with data.
 			// The existing interval should be reduced to no-overlaps, and the new
 			// interval should be inserted.
@@ -283,7 +165,7 @@ TEMPLATE_TEST_CASE("IntervalTree", "", void, int)
 
 		SECTION("Overlap")
 		{
-			IntervalTreeType IntervalTreeB;
+			TIntervalTree<int> IntervalTreeB;
 			Insert(IntervalTreeB, 30, 10, 10);
 			Insert(IntervalTreeB, 50, 10, 11);
 			Insert(IntervalTreeB, 70, 10, 12);
@@ -315,15 +197,15 @@ TEMPLATE_TEST_CASE("IntervalTree", "", void, int)
 
 			for (const FInterval& FreeInterval : Intervals.Free)
 			{
-				REQUIRE(!IntervalTree.Contains(Ptr(FreeInterval.Start), FreeInterval.Size()));
-				REQUIRE(!IntervalTree.Contains(Ptr(FreeInterval.Start), 1));
-				REQUIRE(!IntervalTree.Contains(Ptr(FreeInterval.End - 1), 1));
+				REQUIRE(nullptr == IntervalTree.Get(Ptr(FreeInterval.Start), FreeInterval.Size()));
+				REQUIRE(nullptr == IntervalTree.Get(Ptr(FreeInterval.Start), 1));
+				REQUIRE(nullptr == IntervalTree.Get(Ptr(FreeInterval.End - 1), 1));
 			}
 			for (const FInterval& UsedInterval : Intervals.Used)
 			{
-				REQUIRE(IntervalTree.Contains(Ptr(UsedInterval.Start), UsedInterval.Size()));
-				REQUIRE(IntervalTree.Contains(Ptr(UsedInterval.Start), 1));
-				REQUIRE(IntervalTree.Contains(Ptr(UsedInterval.End - 1), 1));
+				REQUIRE(nullptr != IntervalTree.Get(Ptr(UsedInterval.Start), UsedInterval.Size()));
+				REQUIRE(nullptr != IntervalTree.Get(Ptr(UsedInterval.Start), 1));
+				REQUIRE(nullptr != IntervalTree.Get(Ptr(UsedInterval.End - 1), 1));
 			}
 		}
 	}
@@ -331,96 +213,3 @@ TEMPLATE_TEST_CASE("IntervalTree", "", void, int)
 	IntervalTree.Reset();
 	Validate(IntervalTree, {});
 }
-
-TEST_CASE("IntervalTree.Benchmarks")
-{
-	// Convenience lambda to convert a uintptr_t to a void*
-	auto Ptr = [](uintptr_t Address) { return reinterpret_cast<void*>(Address); };
-
-	AutoRTFM::FIntervalTree IntervalTree;
-
-	// Populates IntervalTree with a random set of intervals.
-	// Returns the FTreeIntervals used to populate the IntervalTree.
-	auto PopulateIntervalTree = [&]() -> FTreeIntervals
-	{
-		const FTreeIntervals Intervals = BuildIntervals(/* Seed */ 1234, /* MaxUsed */ 10'000);
-		for (FInterval Interval : Intervals.Used)
-		{
-			IntervalTree.Insert(Ptr(Interval.Start), Interval.Size());
-		}
-		return Intervals;
-	};
-
-	BENCHMARK_ADVANCED("Insert")(Catch::Benchmark::Chronometer Meter)
-	{
-		const std::vector<FInterval> Intervals = PopulateIntervalTree().Used;
-		Meter.measure([&]
-		{
-			for (FInterval Interval : Intervals)
-			{
-				IntervalTree.Insert(Ptr(Interval.Start), Interval.Size());
-			}
-		});
-	};
-
-	// Populates IntervalTree with a random set of intervals.
-	// If bUsed is true, returns a random set of addresses in those intervals,
-	// otherwise a random set of addresses outside of those intervals.
-	auto PopulateIntervalTreeAndPickAddresses = [&](bool bUsed) -> std::vector<void*>
-	{
-		std::mt19937 Rand(/* Seed */ 4242);
-		
-		const std::vector<FInterval> Intervals = bUsed ? PopulateIntervalTree().Used : PopulateIntervalTree().Free;
-
-		constexpr size_t NumAddresses = 10'000;
-		std::vector<void*> Addresses(NumAddresses);
-		for (size_t I = 0; I < NumAddresses; I++)
-		{
-			FInterval Interval = Intervals[Rand() % Intervals.size()];
-			Addresses[I] = Ptr(Interval.Start + (Rand() % Interval.Size()));
-		}
-		return Addresses;
-	};
-
-	BENCHMARK_ADVANCED("Contains.Used")(Catch::Benchmark::Chronometer Meter)
-	{
-		std::vector<void*> Addresses = PopulateIntervalTreeAndPickAddresses(/* bUsed */ true);
-		Meter.measure([&]
-		{
-			for (void* Address : Addresses)
-			{
-				IntervalTree.Contains(Address, 4);
-			}
-		});
-	};
-
-	BENCHMARK_ADVANCED("Contains.Free")(Catch::Benchmark::Chronometer Meter)
-	{
-		std::vector<void*> Addresses = PopulateIntervalTreeAndPickAddresses(/* bUsed */ false);
-		Meter.measure([&]
-		{
-			for (void* Address : Addresses)
-			{
-				IntervalTree.Contains(Address, 4);
-			}
-		});
-	};
-
-	BENCHMARK_ADVANCED("Merge")(Catch::Benchmark::Chronometer Meter)
-	{
-		PopulateIntervalTree();
-		AutoRTFM::FIntervalTree OtherIntervalTree;
-		{
-			const FTreeIntervals Intervals = BuildIntervals(/* Seed */ 1234, /* MaxUsed */ 10'000);
-			for (FInterval Interval : Intervals.Used)
-			{
-				OtherIntervalTree.Insert(Ptr(Interval.Start), Interval.Size());
-			}
-		}
-
-		Meter.measure([&]
-		{
-			IntervalTree.Merge(OtherIntervalTree);
-		});
-	};
-}
diff --git a/Engine/Source/Programs/AutoRTFMTests/Private/MemoryTrackerTests.cpp b/Engine/Source/Programs/AutoRTFMTests/Private/MemoryTrackerTests.cpp
new file mode 100644
--- /dev/null
+++ b/Engine/Source/Programs/AutoRTFMTests/Private/MemoryTrackerTests.cpp
@@ -0,0 +1,289 @@
+// Copyright Epic Games, Inc. All Rights Reserved.
+
+#include "AutoRTFM.h"
+#include "IntervalTestUtils.h"
+#include "MemoryTracker.h"
+
+#include <random>
+
+using FMemoryTracker = AutoRTFM::FMemoryTracker;
+
+TEST_CASE("MemoryTracker")
+{
+	AutoRTFM::UnreachableIfClosed();
+
+	// Convenience lambda to convert a uintptr_t to a void*
+	auto Ptr = [](uintptr_t Address) { return reinterpret_cast<void*>(Address); };
+
+	auto Insert = [Ptr](FMemoryTracker& Tree, uintptr_t Address, uintptr_t Size)
+	{
+		Tree.Insert(Ptr(Address), Size);
+	};
+
+	// Checks that the tracker has all of the expected intervals, and that Contains() returns
+	// the correct value for various addresses.
+	auto Validate = [&](FMemoryTracker& Tracker, std::vector<FInterval> Expected)
+	{
+		std::vector<FInterval> Got;
+		for (FInterval Entry : Tracker)
+		{
+			Got.push_back(Entry);
+		}
+
+		REQUIRE(Expected.empty() == Tracker.IsEmpty());
+		for (size_t I = 0, N = std::min(Expected.size(), Got.size()); I < N; I++)
+		{
+			REQUIRE(Expected[I].Start == Got[I].Start);
+			REQUIRE(Expected[I].End == Got[I].End);
+		}
+		REQUIRE(Expected.size() == Got.size());
+		for (size_t Index = 0; Index < Expected.size(); Index++)
+		{
+			const FInterval Entry = Expected[Index];
+			const bool bTouchingPrev = Index > 0 && Expected[Index - 1].End == Entry.Start;
+			const bool bTouchingNext = Index < (Expected.size() - 1) && Expected[Index + 1].Start == Entry.End;
+			REQUIRE(true == Tracker.Contains(Ptr(Entry.Start), 1));
+			REQUIRE(true == Tracker.Contains(Ptr(Entry.Start), Entry.Size()));
+			REQUIRE(bTouchingPrev == Tracker.Contains(Ptr(Entry.Start - 1), 1));
+			REQUIRE(bTouchingNext == Tracker.Contains(Ptr(Entry.End), 1));
+		}
+	};
+
+	FMemoryTracker Tracker;
+
+	Validate(Tracker, {});
+
+	SECTION("Insert")
+	{
+		Tracker.Reset();
+		Validate(Tracker, {});
+
+		// New intervals
+		Insert(Tracker, 40, 10);
+		Validate(Tracker, {{40, 50}});
+
+		Insert(Tracker, 20, 10);
+		Validate(Tracker, {{20, 30}, {40, 50}});
+
+		Insert(Tracker, 60, 10);
+		Validate(Tracker, {{20, 30}, {40, 50}, {60, 70}});
+
+		Insert(Tracker, 80, 10);
+		Validate(Tracker, {{20, 30}, {40, 50}, {60, 70}, {80, 90}});
+
+		// Re-inserted intervals
+		Insert(Tracker, 20, 10);
+		Validate(Tracker, {{20, 30}, {40, 50}, {60, 70}, {80, 90}});
+
+		Insert(Tracker, 40, 10);
+		Validate(Tracker, {{20, 30}, {40, 50}, {60, 70}, {80, 90}});
+
+		Insert(Tracker, 60, 10);
+		Validate(Tracker, {{20, 30}, {40, 50}, {60, 70}, {80, 90}});
+
+		Insert(Tracker, 80, 10);
+		Validate(Tracker, {{20, 30}, {40, 50}, {60, 70}, {80, 90}});
+
+		// Expanding intervals without overlap
+		Insert(Tracker, 15, 5);
+		Validate(Tracker, {{15, 30}, {40, 50}, {60, 70}, {80, 90}});
+
+		Insert(Tracker, 30, 5);
+		Validate(Tracker, {{15, 35}, {40, 50}, {60, 70}, {80, 90}});
+
+		Insert(Tracker, 55, 5);
+		Validate(Tracker, {{15, 35}, {40, 50}, {55, 70}, {80, 90}});
+
+		Insert(Tracker, 70, 5);
+		Validate(Tracker, {{15, 35}, {40, 50}, {55, 75}, {80, 90}});
+
+		// Expanding intervals with overlap
+		Insert(Tracker, 35, 5);
+		Validate(Tracker, {{15, 35}, {35, 50}, {55, 75}, {80, 90}});
+
+		Insert(Tracker, 50, 5);
+		Validate(Tracker, {{15, 35}, {35, 55}, {55, 75}, {80, 90}});
+
+		Insert(Tracker, 75, 5);
+		Validate(Tracker, {{15, 35}, {35, 55}, {55, 80}, {80, 90}});
+	}
+
+	SECTION("Merge")
+	{
+		Insert(Tracker, 20, 10);
+		Insert(Tracker, 40, 10);
+		Insert(Tracker, 60, 10);
+		Insert(Tracker, 80, 10);
+		Validate(Tracker, {{20, 30}, {40, 50}, {60, 70}, {80, 90}});
+
+		SECTION("Intertwined")
+		{
+			FMemoryTracker TrackerB;
+			Insert(TrackerB, 10, 5);
+			Insert(TrackerB, 33, 6);
+			Insert(TrackerB, 54, 2);
+			Insert(TrackerB, 73, 3);
+			Insert(TrackerB, 95, 2);
+			Validate(TrackerB, {{10, 15}, {33, 39}, {54, 56}, {73, 76}, {95, 97}});
+
+			Tracker.Merge(TrackerB);
+			Validate(Tracker, {{10, 15}, {20, 30}, {33, 39}, {40, 50}, {54, 56}, {60, 70}, {73, 76}, {80, 90}, {95, 97}});
+		}
+
+		SECTION("Expand")
+		{
+			FMemoryTracker TrackerB;
+			Insert(TrackerB, 15, 5);
+			Insert(TrackerB, 30, 3);
+			Insert(TrackerB, 37, 3);
+			Insert(TrackerB, 50, 2);
+			Insert(TrackerB, 55, 5);
+			Insert(TrackerB, 70, 4);
+			Insert(TrackerB, 78, 2);
+			Insert(TrackerB, 90, 3);
+			Validate(TrackerB, {{15, 20}, {30, 33}, {37, 40}, {50, 52}, {55, 60}, {70, 74}, {78, 80}, {90, 93}});
+
+			Tracker.Merge(TrackerB);
+			Validate(Tracker, {{15, 33}, {37, 52}, {55, 74}, {78, 93}});
+		}
+
+		SECTION("Overlap")
+		{
+			FMemoryTracker TrackerB;
+			Insert(TrackerB, 30, 10);
+			Insert(TrackerB, 50, 10);
+			Insert(TrackerB, 70, 10);
+			Validate(TrackerB, {{30, 40}, {50, 60}, {70, 80}});
+
+			Tracker.Merge(TrackerB);
+			Validate(Tracker, {{20, 30}, {30, 60}, {60, 80}, {80, 90}});
+		}
+	}
+
+	SECTION("Soak")
+	{
+		for (int Pass = 0; Pass < 32; Pass++)
+		{
+			Tracker.Reset();
+			Validate(Tracker, {});
+
+			FTreeIntervals Intervals = BuildIntervals(/* Seed */ Pass, /* MaxUsed */ 1000);
+
+			for (FInterval Interval : Intervals.Used)
+			{
+				Insert(Tracker, Interval.Start, Interval.Size());
+			}
+
+			Intervals.SortAndValidate();
+
+			for (const FInterval& FreeInterval : Intervals.Free)
+			{
+				REQUIRE(!Tracker.Contains(Ptr(FreeInterval.Start), FreeInterval.Size()));
+				REQUIRE(!Tracker.Contains(Ptr(FreeInterval.Start), 1));
+				REQUIRE(!Tracker.Contains(Ptr(FreeInterval.End - 1), 1));
+			}
+			for (const FInterval& UsedInterval : Intervals.Used)
+			{
+				REQUIRE(Tracker.Contains(Ptr(UsedInterval.Start), UsedInterval.Size()));
+				REQUIRE(Tracker.Contains(Ptr(UsedInterval.Start), 1));
+				REQUIRE(Tracker.Contains(Ptr(UsedInterval.End - 1), 1));
+			}
+		}
+	}
+
+	Tracker.Reset();
+	Validate(Tracker, {});
+}
+
+TEST_CASE("MemoryTracker.Benchmarks")
+{
+	// Convenience lambda to convert a uintptr_t to a void*
+	auto Ptr = [](uintptr_t Address) { return reinterpret_cast<void*>(Address); };
+
+	FMemoryTracker Tracker;
+
+	// Populates Tracker with a random set of intervals.
+	// Returns the FTreeIntervals used to populate the tracker.
+	auto PopulateTracker = [&]() -> FTreeIntervals
+	{
+		const FTreeIntervals Intervals = BuildIntervals(/* Seed */ 1234, /* MaxUsed */ 10'000);
+		for (FInterval Interval : Intervals.Used)
+		{
+			Tracker.Insert(Ptr(Interval.Start), Interval.Size());
+		}
+		return Intervals;
+	};
+
+	BENCHMARK_ADVANCED("Insert")(Catch::Benchmark::Chronometer Meter)
+	{
+		const std::vector<FInterval> Intervals = PopulateTracker().Used;
+		Meter.measure([&]
+		{
+			for (FInterval Interval : Intervals)
+			{
+				Tracker.Insert(Ptr(Interval.Start), Interval.Size());
+			}
+		});
+	};
+
+	// Populates Tracker with a random set of intervals.
+	// If bUsed is true, returns a random set of addresses in those intervals,
+	// otherwise a random set of addresses outside of those intervals.
+	auto PopulateTrackerAndPickAddresses = [&](bool bUsed) -> std::vector<void*>
+	{
+		std::mt19937 Rand(/* Seed */ 4242);
+
+		const std::vector<FInterval> Intervals = bUsed ? PopulateTracker().Used : PopulateTracker().Free;
+
+		constexpr size_t NumAddresses = 10'000;
+		std::vector<void*> Addresses(NumAddresses);
+		for (size_t I = 0; I < NumAddresses; I++)
+		{
+			FInterval Interval = Intervals[Rand() % Intervals.size()];
+			Addresses[I] = Ptr(Interval.Start + (Rand() % Interval.Size()));
+		}
+		return Addresses;
+	};
+
+	BENCHMARK_ADVANCED("Contains.Used")(Catch::Benchmark::Chronometer Meter)
+	{
+		std::vector<void*> Addresses = PopulateTrackerAndPickAddresses(/* bUsed */ true);
+		Meter.measure([&]
+		{
+			for (void* Address : Addresses)
+			{
+				Tracker.Contains(Address, 4);
+			}
+		});
+	};
+
+	BENCHMARK_ADVANCED("Contains.Free")(Catch::Benchmark::Chronometer Meter)
+	{
+		std::vector<void*> Addresses = PopulateTrackerAndPickAddresses(/* bUsed */ false);
+		Meter.measure([&]
+		{
+			for (void* Address : Addresses)
+			{
+				Tracker.Contains(Address, 4);
+			}
+		});
+	};
+
+	BENCHMARK_ADVANCED("Merge")(Catch::Benchmark::Chronometer Meter)
+	{
+		PopulateTracker();
+		FMemoryTracker OtherTracker;
+		{
+			const FTreeIntervals Intervals = BuildIntervals(/* Seed */ 1234, /* MaxUsed */ 10'000);
+			for (FInterval Interval : Intervals.Used)
+			{
+				OtherTracker.Insert(Ptr(Interval.Start), Interval.Size());
+			}
+		}
+
+		Meter.measure([&]
+		{
+			Tracker.Merge(OtherTracker);
+		});
+	};
+}
diff --git a/Engine/Source/Runtime/AutoRTFM/Private/Interval.h b/Engine/Source/Runtime/AutoRTFM/Private/Interval.h
new file mode 100644
--- /dev/null
+++ b/Engine/Source/Runtime/AutoRTFM/Private/Interval.h
@@ -0,0 +1,44 @@
+// Copyright Epic Games, Inc. All Rights Reserved.
+
+#pragma once
+
+#if (defined(__AUTORTFM) && __AUTORTFM)
+
+#include <cstddef>
+#include <cstdint>
+
+namespace AutoRTFM
+{
+
+// A half-open address interval [Start, End).
+struct FInterval final
+{
+	// Constructs an interval with explicit start and end addresses.
+	constexpr FInterval(uintptr_t Start, uintptr_t End) : Start{Start}, End{End} {}
+
+	// Constructs an interval from a base address and byte size.
+	FInterval(const void* const Address, const size_t Size)
+		: Start(reinterpret_cast<uintptr_t>(Address))
+		, End(reinterpret_cast<uintptr_t>(Address) + Size)
+	{
+	}
+
+	// Returns the number of bytes in the interval.
+	size_t Size() const
+	{
+		return End - Start;
+	}
+
+	// Orders intervals by start address, for use in sorted containers.
+	bool operator<(const FInterval& Other) const
+	{
+		return Start < Other.Start;
+	}
+
+	uintptr_t Start;
+	uintptr_t End;
+};
+
+}  // namespace AutoRTFM
+
+#endif  // (defined(__AUTORTFM) && __AUTORTFM)
diff --git a/Engine/Source/Runtime/AutoRTFM/Private/IntervalTree.h b/Engine/Source/Runtime/AutoRTFM/Private/IntervalTree.h
--- a/Engine/Source/Runtime/AutoRTFM/Private/IntervalTree.h
+++ b/Engine/Source/Runtime/AutoRTFM/Private/IntervalTree.h
@@ -6,45 +6,25 @@
 
 #include "Allocator.h"
 #include "ContainerValidation.h"
+#include "Interval.h"
 #include "Stack.h"
 #include "Utils.h"
 
 #include <cstring>
-#include <type_traits>
 
 namespace AutoRTFM
 {
 
-struct FInterval final
-{
-	constexpr FInterval(uintptr_t Start, uintptr_t End) : Start{Start}, End{End} {}
-
-	FInterval(const void* const Address, const size_t Size)
-		: Start(reinterpret_cast<uintptr_t>(Address))
-		, End(reinterpret_cast<uintptr_t>(Address) + Size)
-	{
-	}
-
-	size_t Size() const
-	{
-		return End - Start;
-	}
-	bool operator<(const FInterval& Other) const
-	{
-		return Start < Other.Start;
-	}
-
-	uintptr_t Start;
-	uintptr_t End;
-};
-
+// TIntervalTree tracks a set of non-overlapping memory intervals, each associated with a
+// data value of type DataTypeIn, using a red-black tree. Intervals are stored as [Start, End)
+// and intervals that touch at their endpoints are merged into a single interval on insert,
+// retaining the data of the existing interval.
 template <typename DataTypeIn, Allocator Allocator = FAllocator>
 struct AUTORTFM_DISABLE_NO_SANITIZE TIntervalTree final
 {
 private:
-	static constexpr bool bHasData = !std::is_same_v<DataTypeIn, void>;
 	using NodeIndex = uint32_t;
-	using DataType = std::conditional_t<bHasData, DataTypeIn, std::monostate>;
+	using DataType = DataTypeIn;
 	static constexpr NodeIndex InvalidNode = UINT32_MAX;
 
 public:
@@ -53,54 +33,45 @@ public:
 	TIntervalTree(const TIntervalTree&) = delete;
 	TIntervalTree& operator=(const TIntervalTree&) = delete;
 
-	bool Insert(const void* const Address, const size_t Size)
-	requires(!bHasData)
-	{
-		return Insert(FInsertArgs{
-			.Interval = FInterval{Address, Size}
-        });
-	}
-
+	// Inserts the interval [Address, Address+Size) with the given data into the tree.
+	// If the new interval is adjacent to an existing interval, the two are merged and
+	// the existing interval's data is retained.
+	// Returns false if the interval overlaps (but is not adjacent to) an existing interval.
 	bool Insert(const void* const Address, const size_t Size, DataType&& Data)
-	requires(bHasData)
 	{
-		return Insert(FInsertArgs{
-			.Interval = FInterval{Address, Size},
-              .Data = std::move(Data)
-        });
+		return Insert(FInterval{Address, Size}, std::move(Data));
 	}
 
+	// Inserts the interval [Address, Address+Size) with the given data into the tree.
+	// If the new interval is adjacent to an existing interval, the two are merged and
+	// the existing interval's data is retained.
+	// Returns false if the interval overlaps (but is not adjacent to) an existing interval.
 	bool Insert(const void* const Address, const size_t Size, const DataType& Data)
-	requires(bHasData)
 	{
-		return Insert(FInsertArgs{
-			.Interval = FInterval{Address, Size},
-              .Data = Data
-        });
-	}
-
-	bool Contains(const void* const Address, const size_t Size) const
-	{
-		return Query(Address, Size);
+		return Insert(FInterval{Address, Size}, Data);
 	}
 
+	// Returns a pointer to the data associated with the stored interval that contains
+	// [Address, Address+Size), or nullptr if no such interval exists.
 	const DataType* Get(const void* const Address, const size_t Size) const
-	requires(bHasData)
 	{
 		return Query(Address, Size);
 	}
 
+	// Returns true if no intervals are stored.
 	bool IsEmpty() const
 	{
 		return InvalidNode == Root;
 	}
 
+	// Removes all stored intervals.
 	void Reset()
 	{
 		Root = InvalidNode;
 		Nodes.Reset();
 	}
 
+	// Inserts all intervals from Other into this tree.
 	void Merge(const TIntervalTree& Other)
 	{
 		if (Other.IsEmpty())
@@ -117,17 +88,7 @@ public:
 			const NodeIndex Current = ToProcess.Back();
 			ToProcess.Pop();
 
-			if constexpr (bHasData)
-			{
-				Insert({
-					.Interval = Other.Nodes.Interval(Current),
-					.Data = std::move(Other.Nodes.Data(Current)),
-				});
-			}
-			else
-			{
-				Insert({.Interval = Other.Nodes.Interval(Current)});
-			}
+			Insert(Other.Nodes.Interval(Current), Other.Nodes.Data(Current));
 
 			const NodeIndex Left = Other.Nodes.Left(Current);
 			const NodeIndex Right = Other.Nodes.Right(Current);
@@ -144,26 +105,20 @@ public:
 		} while (!ToProcess.IsEmpty());
 	}
 
+	// A single entry yielded by the iterator: the stored interval and a reference to its data.
 	struct FIntervalAndData
 	{
 		FInterval Interval;
 		const DataType& Data;
 	};
 
-	// A forward iterator of the TIntervalTree
+	// A forward iterator that yields intervals in ascending order of start address.
 	struct FIterator
 	{
 		FIterator(const FIterator&) = default;
 		FIterator& operator=(const FIterator&) = default;
 
-		FInterval operator*() const
-		requires(!bHasData)
-		{
-			return Tree->Nodes.Interval(Node);
-		}
-
 		FIntervalAndData operator*() const
-		requires(bHasData)
 		{
 			return {.Interval = Tree->Nodes.Interval(Node), .Data = Tree->Nodes.Data(Node)};
 		}
@@ -217,22 +172,16 @@ private:
 		DataType Data;
 	};
 
-	struct FNodePtrsNoData
+	struct FNodePtrs final
 	{
 		FInterval* Interval = nullptr;
 		NodeIndex* Left = nullptr;
 		NodeIndex* Right = nullptr;
 		NodeIndex* Parent = nullptr;
 		ColorWord* Color = nullptr;
-	};
-
-	struct FNodePtrsWithData final : FNodePtrsNoData
-	{
 		DataType* Data = nullptr;
 	};
 
-	using FNodePtrs = std::conditional_t<bHasData, FNodePtrsWithData, FNodePtrsNoData>;
-
 	// A list of FNodes, stored as a structure-of-arrays.
 	struct FNodes final
 	{
@@ -253,12 +202,9 @@ private:
 		// Clears the node list to an empty state.
 		void Reset()
 		{
-			if constexpr (bHasData)
+			for (size_t I = 0; I < Count; I++)
 			{
-				for (size_t I = 0; I < Count; I++)
-				{
-					Ptrs.Data[I].~DataType();
-				}
+				Ptrs.Data[I].~DataType();
 			}
 
 			Ptrs = {};
@@ -284,10 +230,7 @@ private:
 			Ptrs.Right[Count] = Node.Right;
 			Ptrs.Parent[Count] = Node.Parent;
 			SetColor(Count, Node.Color);
-			if constexpr (bHasData)
-			{
-				new (&Ptrs.Data[Count]) DataType(std::move(Node.Data));
-			}
+			new (&Ptrs.Data[Count]) DataType(std::move(Node.Data));
 
 			return static_cast<NodeIndex>(Count++);
 		}
@@ -325,12 +268,10 @@ private:
 			return Ptrs.Parent[Index];
 		}
 		DataType& Data(NodeIndex Index)
-		requires(bHasData)
 		{
 			return Ptrs.Data[Index];
 		}
 		const DataType& Data(NodeIndex Index) const
-		requires(bHasData)
 		{
 			return Ptrs.Data[Index];
 		}
@@ -371,14 +312,14 @@ private:
 			size_t const RightOffset = SubAllocate<NodeIndex>(TotalSize, NewCapacity);
 			size_t const ParentOffset = SubAllocate<NodeIndex>(TotalSize, NewCapacity);
 			size_t const ColorOffset = SubAllocate<ColorWord>(TotalSize, NumColorWords(NewCapacity));
-			size_t const DataOffset = bHasData ? SubAllocate<DataType>(TotalSize, NewCapacity) : 0;
+			size_t const DataOffset = SubAllocate<DataType>(TotalSize, NewCapacity);
 			auto* const NewMemory = reinterpret_cast<std::byte*>(Allocator::Allocate(TotalSize, alignof(FNode)));
 			auto* const NewInterval = std::launder(reinterpret_cast<FInterval*>(NewMemory + IntervalOffset));
 			auto* const NewLeft = std::launder(reinterpret_cast<NodeIndex*>(NewMemory + LeftOffset));
 			auto* const NewRight = std::launder(reinterpret_cast<NodeIndex*>(NewMemory + RightOffset));
 			auto* const NewParent = std::launder(reinterpret_cast<NodeIndex*>(NewMemory + ParentOffset));
 			auto* const NewColor = std::launder(reinterpret_cast<ColorWord*>(NewMemory + ColorOffset));
-			auto* const NewData = bHasData ? std::launder(reinterpret_cast<DataType*>(NewMemory + DataOffset)) : nullptr;
+			auto* const NewData = std::launder(reinterpret_cast<DataType*>(NewMemory + DataOffset));
 			if (size_t const NumToCopy = std::min(Count, NewCapacity))
 			{
 				memcpy(NewInterval, Ptrs.Interval, sizeof(FInterval) * NumToCopy);
@@ -387,13 +328,10 @@ private:
 				memcpy(NewParent, Ptrs.Parent, sizeof(NodeIndex) * NumToCopy);
 				memcpy(NewColor, Ptrs.Color, sizeof(ColorWord) * NumColorWords(NumToCopy));
 
-				if constexpr (bHasData)
+				for (size_t I = 0; I < NumToCopy; I++)
 				{
-					for (size_t I = 0; I < NumToCopy; I++)
-					{
-						new (&NewData[I]) DataType(std::move(Ptrs.Data[I]));
-						Ptrs.Data[I].~DataType();
-					}
+					new (&NewData[I]) DataType(std::move(Ptrs.Data[I]));
+					Ptrs.Data[I].~DataType();
 				}
 			}
 
@@ -404,10 +342,7 @@ private:
 			Ptrs.Right = NewRight;
 			Ptrs.Parent = NewParent;
 			Ptrs.Color = NewColor;
-			if constexpr (bHasData)
-			{
-				Ptrs.Data = NewData;
-			}
+			Ptrs.Data = NewData;
 
 			Capacity = NewCapacity;
 			Memory = NewMemory;
@@ -466,45 +401,22 @@ private:
 		return Node;
 	}
 
-	struct FInsertArgsNoData
-	{
-		FInterval Interval;
-	};
-	struct FInsertArgsWithData
+	UE_AUTORTFM_FORCEINLINE NodeIndex AddNode(FInterval NewInterval, DataType&& Data, NodeIndex Parent = InvalidNode)
 	{
-		FInterval Interval;
-		DataType Data;
-	};
-	using FInsertArgs = std::conditional_t<bHasData, FInsertArgsWithData, FInsertArgsNoData>;
-
-	UE_AUTORTFM_FORCEINLINE NodeIndex AddNode(FInsertArgs&& Args, NodeIndex Parent = InvalidNode)
-	{
-		if constexpr (bHasData)
-		{
-			return Nodes.Push({
-				.Interval = Args.Interval,
-				.Parent = Parent,
-				.Color = (Parent == InvalidNode) ? EColor::Black : EColor::Red,
-				.Data = Args.Data,
-			});
-		}
-		else
-		{
-			return Nodes.Push({
-				.Interval = Args.Interval,
-				.Parent = Parent,
-				.Color = (Parent == InvalidNode) ? EColor::Black : EColor::Red,
-				.Data = {},
-			});
-		}
+		return Nodes.Push({
+			.Interval = NewInterval,
+			.Parent = Parent,
+			.Color = (Parent == InvalidNode) ? EColor::Black : EColor::Red,
+			.Data = std::move(Data),
+		});
 	}
 
-	UE_AUTORTFM_FORCENOINLINE bool Insert(FInsertArgs&& Args)
+	UE_AUTORTFM_FORCENOINLINE bool Insert(FInterval NewInterval, DataType Data)
 	{
 		if (AUTORTFM_UNLIKELY(InvalidNode == Root))
 		{
 			AUTORTFM_ASSERT(Nodes.Count == 0);
-			Root = AddNode(std::move(Args));
+			Root = AddNode(NewInterval, std::move(Data));
 			return true;
 		}
 
@@ -514,24 +426,24 @@ private:
 		{
 			const FInterval Interval = Nodes.Interval(Current);
 
-			if (AUTORTFM_UNLIKELY((Args.Interval.Start < Interval.End) && (Interval.Start < Args.Interval.End)))
+			if (AUTORTFM_UNLIKELY((NewInterval.Start < Interval.End) && (Interval.Start < NewInterval.End)))
 			{
 				return false;
 			}
 
-			if (Args.Interval.Start < Interval.Start)
+			if (NewInterval.Start < Interval.Start)
 			{
-				if (AUTORTFM_UNLIKELY(Args.Interval.End == Interval.Start))
+				if (AUTORTFM_UNLIKELY(NewInterval.End == Interval.Start))
 				{
-					AUTORTFM_ASSERT(Args.Interval.Start < Interval.Start);
+					AUTORTFM_ASSERT(NewInterval.Start < Interval.Start);
 
 					// We can just modify the existing node in place.
-					Nodes.Interval(Current).Start = Args.Interval.Start;
+					Nodes.Interval(Current).Start = NewInterval.Start;
 					return true;
 				}
 				else if (InvalidNode == Nodes.Left(Current))
 				{
-					const NodeIndex Index = AddNode(std::move(Args), Current);
+					const NodeIndex Index = AddNode(NewInterval, std::move(Data), Current);
 					Nodes.Left(Current) = Index;
 					Current = Index;
 					break;
@@ -541,17 +453,17 @@ private:
 			}
 			else
 			{
-				if (AUTORTFM_UNLIKELY(Args.Interval.Start == Interval.End))
+				if (AUTORTFM_UNLIKELY(NewInterval.Start == Interval.End))
 				{
-					AUTORTFM_ASSERT(Args.Interval.End > Interval.End);
+					AUTORTFM_ASSERT(NewInterval.End > Interval.End);
 
 					// We can just modify the existing node in place.
-					Nodes.Interval(Current).End = Args.Interval.End;
+					Nodes.Interval(Current).End = NewInterval.End;
 					return true;
 				}
 				else if (InvalidNode == Nodes.Right(Current))
 				{
-					const NodeIndex Index = AddNode(std::move(Args), Current);
+					const NodeIndex Index = AddNode(NewInterval, std::move(Data), Current);
 					Nodes.Right(Current) = Index;
 					Current = Index;
 					break;
@@ -711,12 +623,11 @@ private:
 		return true;
 	}
 
-	UE_AUTORTFM_FORCENOINLINE
-	std::conditional_t<bHasData, const DataType*, bool> Query(const void* const Address, const size_t Size) const
+	UE_AUTORTFM_FORCENOINLINE const DataType* Query(const void* const Address, const size_t Size) const
 	{
 		if (AUTORTFM_UNLIKELY(InvalidNode == Root))
 		{
-			return {};
+			return nullptr;
 		}
 
 		const FInterval QueryInterval{Address, Size};
@@ -727,26 +638,16 @@ private:
 		{
 			const FInterval Interval = Nodes.Interval(Current);
 
-			// This check does not need to prove that QueryInterval is entirely
-			// enclosed within Interval, because if any byte of QueryInterval was in the
-			// original Interval then it **must** already have been new memory.
 			if ((QueryInterval.Start < Interval.End) && (Interval.Start < QueryInterval.End))
 			{
-				if constexpr (bHasData)
-				{
-					return &Nodes.Data(Current);
-				}
-				else
-				{
-					return true;
-				}
+				return &Nodes.Data(Current);
 			}
 
 			const NodeIndex* Next = (QueryInterval.Start < Interval.End) ? Nodes.Ptrs.Left : Nodes.Ptrs.Right;
 			Current = Next[Current];
 		} while (InvalidNode != Current);
 
-		return {};
+		return nullptr;
 	}
 
 	UE_AUTORTFM_FORCEINLINE void AssertStructureIsOk() const
@@ -799,9 +700,6 @@ private:
 	}
 };
 
-// FIntervalTree is an TIntervalTree with no data
-using FIntervalTree = TIntervalTree<void>;
-
 }  // namespace AutoRTFM
 
 #endif  // (defined(__AUTORTFM) && __AUTORTFM)
diff --git a/Engine/Source/Runtime/AutoRTFM/Private/MemoryTracker.h b/Engine/Source/Runtime/AutoRTFM/Private/MemoryTracker.h
new file mode 100644
--- /dev/null
+++ b/Engine/Source/Runtime/AutoRTFM/Private/MemoryTracker.h
@@ -0,0 +1,665 @@
+// Copyright Epic Games, Inc. All Rights Reserved.
+
+#pragma once
+
+#if (defined(__AUTORTFM) && __AUTORTFM)
+
+#include "Allocator.h"
+#include "ContainerValidation.h"
+#include "Interval.h"
+#include "Stack.h"
+#include "Utils.h"
+
+#include <cstring>
+
+namespace AutoRTFM
+{
+
+// FMemoryTracker tracks a set of non-overlapping memory intervals using a red-black tree.
+// Each interval is stored as [Start, End) and intervals that touch at their endpoints are
+// merged into a single interval automatically on insert.
+//
+// Assumes that queries never partially overlap a stored interval — a query either falls
+// entirely within an interval or entirely outside it. This assumption enables a simplified
+// BST traversal in Contains() that avoids examining both subtrees.
+struct AUTORTFM_DISABLE_NO_SANITIZE FMemoryTracker final
+{
+private:
+	using NodeIndex = uint32_t;
+	static constexpr NodeIndex InvalidNode = UINT32_MAX;
+
+public:
+	FMemoryTracker() = default;
+
+	FMemoryTracker(const FMemoryTracker&) = delete;
+	FMemoryTracker& operator=(const FMemoryTracker&) = delete;
+
+	// Inserts the interval [Address, Address+Size) into the tracker.
+	// If the new interval is adjacent to an existing interval, the two are merged.
+	// Returns false if the interval overlaps (but is not adjacent to) an existing interval.
+	bool Insert(const void* const Address, const size_t Size)
+	{
+		return Insert(FInterval{Address, Size});
+	}
+
+	// Returns true if the interval [Address, Address+Size) overlaps any stored interval.
+	bool Contains(const void* const Address, const size_t Size) const
+	{
+		return Query(Address, Size);
+	}
+
+	// Returns true if no intervals are stored.
+	bool IsEmpty() const
+	{
+		return InvalidNode == Root;
+	}
+
+	// Removes all stored intervals.
+	void Reset()
+	{
+		Root = InvalidNode;
+		Nodes.Reset();
+	}
+
+	// Inserts all intervals from Other into this tracker.
+	void Merge(const FMemoryTracker& Other)
+	{
+		if (Other.IsEmpty())
+		{
+			return;
+		}
+
+		TStack<NodeIndex, 32, FAllocator, EContainerValidation::Disabled> ToProcess;
+
+		ToProcess.Push(Other.Root);
+
+		do
+		{
+			const NodeIndex Current = ToProcess.Back();
+			ToProcess.Pop();
+
+			Insert(Other.Nodes.Interval(Current));
+
+			const NodeIndex Left = Other.Nodes.Left(Current);
+			const NodeIndex Right = Other.Nodes.Right(Current);
+
+			if (InvalidNode != Left)
+			{
+				ToProcess.Push(Left);
+			}
+
+			if (InvalidNode != Right)
+			{
+				ToProcess.Push(Right);
+			}
+		} while (!ToProcess.IsEmpty());
+	}
+
+	// A forward iterator of the FMemoryTracker
+	struct FIterator
+	{
+		FIterator(const FIterator&) = default;
+		FIterator& operator=(const FIterator&) = default;
+
+		FInterval operator*() const
+		{
+			return Tree->Nodes.Interval(Node);
+		}
+
+		FIterator& operator++()
+		{
+			Node = Tree->Next(Node);
+			return *this;
+		}
+
+		bool operator==(const FIterator& Other) const = default;
+		bool operator!=(const FIterator& Other) const = default;
+
+	private:
+		friend struct FMemoryTracker;
+		FIterator(FMemoryTracker const* Tree, NodeIndex Node) : Tree{Tree}, Node{Node} {}
+		FMemoryTracker const* Tree = nullptr;
+		NodeIndex Node = InvalidNode;
+	};
+
+	// Returns an iterator to the first interval in the tree
+	FIterator begin() const
+	{
+		return FIterator(this, Root == InvalidNode ? InvalidNode : LeftMostDescendant(Root));
+	}
+
+	// Returns an end iterator for the tree
+	FIterator end() const
+	{
+		return FIterator(this, InvalidNode);
+	}
+
+private:
+	friend struct FIterator;
+
+	enum EColor : bool
+	{
+		Red = false,
+		Black = true
+	};
+
+	using ColorWord = uint64_t;
+
+	struct FNode final
+	{
+		FInterval Interval;
+		NodeIndex Left = InvalidNode;
+		NodeIndex Right = InvalidNode;
+		NodeIndex Parent = InvalidNode;
+		EColor Color;
+	};
+
+	struct FNodePtrs final
+	{
+		FInterval* Interval = nullptr;
+		NodeIndex* Left = nullptr;
+		NodeIndex* Right = nullptr;
+		NodeIndex* Parent = nullptr;
+		ColorWord* Color = nullptr;
+	};
+
+	// A list of FNodes, stored as a structure-of-arrays.
+	struct FNodes final
+	{
+		static constexpr size_t NumColorWordBits = sizeof(ColorWord) * 8;
+		static constexpr NodeIndex InitialCapacity = 32;
+
+		FNodePtrs Ptrs;
+		NodeIndex Count = 0;
+		NodeIndex Capacity = 0;
+
+		FNodes() = default;
+
+		~FNodes()
+		{
+			Reset();
+		}
+
+		// Clears the node list to an empty state.
+		void Reset()
+		{
+			Ptrs = {};
+			Count = 0;
+			Capacity = 0;
+
+			FAllocator::Free(Memory, MemorySizeInBytes);
+			Memory = nullptr;
+			MemorySizeInBytes = 0;
+		}
+
+		// Pushes a new node to the end of the list.
+		// Returns the node index.
+		NodeIndex Push(FNode Node)
+		{
+			if (Count >= Capacity)
+			{
+				Resize(std::max(Capacity * 2, InitialCapacity));
+			}
+
+			Ptrs.Interval[Count] = Node.Interval;
+			Ptrs.Left[Count] = Node.Left;
+			Ptrs.Right[Count] = Node.Right;
+			Ptrs.Parent[Count] = Node.Parent;
+			SetColor(Count, Node.Color);
+
+			return static_cast<NodeIndex>(Count++);
+		}
+
+		FInterval& Interval(NodeIndex Index)
+		{
+			return Ptrs.Interval[Index];
+		}
+		const FInterval& Interval(NodeIndex Index) const
+		{
+			return Ptrs.Interval[Index];
+		}
+		NodeIndex& Left(NodeIndex Index)
+		{
+			return Ptrs.Left[Index];
+		}
+		const NodeIndex& Left(NodeIndex Index) const
+		{
+			return Ptrs.Left[Index];
+		}
+		NodeIndex& Right(NodeIndex Index)
+		{
+			return Ptrs.Right[Index];
+		}
+		const NodeIndex& Right(NodeIndex Index) const
+		{
+			return Ptrs.Right[Index];
+		}
+		NodeIndex& Parent(NodeIndex Index)
+		{
+			return Ptrs.Parent[Index];
+		}
+		const NodeIndex& Parent(NodeIndex Index) const
+		{
+			return Ptrs.Parent[Index];
+		}
+
+		void SetColor(NodeIndex Index, EColor Color)
+		{
+			ColorWord const Bit = static_cast<ColorWord>(1) << (Index & (NumColorWordBits - 1));
+			ColorWord& Word = Ptrs.Color[Index / NumColorWordBits];
+			if (Color == EColor::Black)
+			{
+				Word |= Bit;
+			}
+			else
+			{
+				Word &= ~Bit;
+			}
+		}
+
+		EColor Color(size_t Index) const
+		{
+			ColorWord const Bit = static_cast<ColorWord>(1) << (Index & (NumColorWordBits - 1));
+			ColorWord& Word = Ptrs.Color[Index / NumColorWordBits];
+			return (Word & Bit) ? EColor::Black : EColor::Red;
+		}
+
+	private:
+		std::byte* Memory = nullptr;
+		size_t MemorySizeInBytes = 0;
+
+		FNodes(const FNodes&) = delete;
+		FNodes& operator=(const FNodes&) = delete;
+
+		void Resize(NodeIndex NewCapacity)
+		{
+			size_t TotalSize = 0;
+			size_t const IntervalOffset = SubAllocate<FInterval>(TotalSize, NewCapacity);
+			size_t const LeftOffset = SubAllocate<NodeIndex>(TotalSize, NewCapacity);
+			size_t const RightOffset = SubAllocate<NodeIndex>(TotalSize, NewCapacity);
+			size_t const ParentOffset = SubAllocate<NodeIndex>(TotalSize, NewCapacity);
+			size_t const ColorOffset = SubAllocate<ColorWord>(TotalSize, NumColorWords(NewCapacity));
+			auto* const NewMemory = reinterpret_cast<std::byte*>(FAllocator::Allocate(TotalSize, alignof(FNode)));
+			auto* const NewInterval = std::launder(reinterpret_cast<FInterval*>(NewMemory + IntervalOffset));
+			auto* const NewLeft = std::launder(reinterpret_cast<NodeIndex*>(NewMemory + LeftOffset));
+			auto* const NewRight = std::launder(reinterpret_cast<NodeIndex*>(NewMemory + RightOffset));
+			auto* const NewParent = std::launder(reinterpret_cast<NodeIndex*>(NewMemory + ParentOffset));
+			auto* const NewColor = std::launder(reinterpret_cast<ColorWord*>(NewMemory + ColorOffset));
+			if (size_t const NumToCopy = std::min(Count, NewCapacity))
+			{
+				memcpy(NewInterval, Ptrs.Interval, sizeof(FInterval) * NumToCopy);
+				memcpy(NewLeft, Ptrs.Left, sizeof(NodeIndex) * NumToCopy);
+				memcpy(NewRight, Ptrs.Right, sizeof(NodeIndex) * NumToCopy);
+				memcpy(NewParent, Ptrs.Parent, sizeof(NodeIndex) * NumToCopy);
+				memcpy(NewColor, Ptrs.Color, sizeof(ColorWord) * NumColorWords(NumToCopy));
+			}
+
+			FAllocator::Free(Memory, MemorySizeInBytes);
+
+			Ptrs.Interval = NewInterval;
+			Ptrs.Left = NewLeft;
+			Ptrs.Right = NewRight;
+			Ptrs.Parent = NewParent;
+			Ptrs.Color = NewColor;
+
+			Capacity = NewCapacity;
+			Memory = NewMemory;
+			MemorySizeInBytes = TotalSize;
+		}
+
+		template <typename T>
+		static size_t SubAllocate(size_t& Offset, size_t Count)
+		{
+			size_t Base = AlignUp(Offset, alignof(T));
+			Offset = Base + sizeof(T) * Count;
+			return Base;
+		}
+
+		static constexpr size_t NumColorWords(size_t Count)
+		{
+			return (Count + NumColorWordBits - 1) / NumColorWordBits;
+		}
+	};
+
+	NodeIndex Root = InvalidNode;
+	FNodes Nodes;
+
+	static constexpr bool bExtraDebugging = false;
+
+	// Returns the index of the next sequential node after Node.
+	NodeIndex Next(NodeIndex Node) const
+	{
+		// If Node has a right child, then the next node is the left-most
+		// descendant of that node (or that node, if it has no descendants).
+		if (NodeIndex Right = Nodes.Right(Node); Right != InvalidNode)
+		{
+			return LeftMostDescendant(Right);
+		}
+
+		// Otherwise the next node is the first ancestor where Node is under its left branch.
+		for (NodeIndex Parent = Nodes.Parent(Node); Parent != InvalidNode; Parent = Nodes.Parent(Node))
+		{
+			if (Nodes.Left(Parent) == Node)
+			{
+				return Parent;
+			}
+			Node = Parent;
+		}
+
+		return InvalidNode;
+	}
+
+	// Returns the left-most leaf descendant of Node.
+	NodeIndex LeftMostDescendant(NodeIndex Node) const
+	{
+		for (NodeIndex Left = Nodes.Left(Node); Left != InvalidNode; Left = Nodes.Left(Node))
+		{
+			Node = Left;
+		}
+		return Node;
+	}
+
+	UE_AUTORTFM_FORCEINLINE NodeIndex AddNode(FInterval Interval, NodeIndex Parent = InvalidNode)
+	{
+		return Nodes.Push({
+			.Interval = Interval,
+			.Parent = Parent,
+			.Color = (Parent == InvalidNode) ? EColor::Black : EColor::Red,
+		});
+	}
+
+	UE_AUTORTFM_FORCENOINLINE bool Insert(FInterval NewInterval)
+	{
+		if (AUTORTFM_UNLIKELY(InvalidNode == Root))
+		{
+			AUTORTFM_ASSERT(Nodes.Count == 0);
+			Root = AddNode(NewInterval);
+			return true;
+		}
+
+		NodeIndex Current = Root;
+
+		for (;;)
+		{
+			const FInterval Interval = Nodes.Interval(Current);
+
+			if (AUTORTFM_UNLIKELY((NewInterval.Start < Interval.End) && (Interval.Start < NewInterval.End)))
+			{
+				return false;
+			}
+
+			if (NewInterval.Start < Interval.Start)
+			{
+				if (AUTORTFM_UNLIKELY(NewInterval.End == Interval.Start))
+				{
+					AUTORTFM_ASSERT(NewInterval.Start < Interval.Start);
+
+					// We can just modify the existing node in place.
+					Nodes.Interval(Current).Start = NewInterval.Start;
+					return true;
+				}
+				else if (InvalidNode == Nodes.Left(Current))
+				{
+					const NodeIndex Index = AddNode(NewInterval, Current);
+					Nodes.Left(Current) = Index;
+					Current = Index;
+					break;
+				}
+
+				Current = Nodes.Left(Current);
+			}
+			else
+			{
+				if (AUTORTFM_UNLIKELY(NewInterval.Start == Interval.End))
+				{
+					AUTORTFM_ASSERT(NewInterval.End > Interval.End);
+
+					// We can just modify the existing node in place.
+					Nodes.Interval(Current).End = NewInterval.End;
+					return true;
+				}
+				else if (InvalidNode == Nodes.Right(Current))
+				{
+					const NodeIndex Index = AddNode(NewInterval, Current);
+					Nodes.Right(Current) = Index;
+					Current = Index;
+					break;
+				}
+
+				Current = Nodes.Right(Current);
+			}
+
+			AUTORTFM_ASSERT(Root != Current);
+		}
+
+		auto IsBlack = [this](NodeIndex Index)
+		{
+			return (InvalidNode == Index) || (EColor::Black == Nodes.Color(Index));
+		};
+
+		for (;;)
+		{
+			NodeIndex Parent = Nodes.Parent(Current);
+
+			UE_AUTORTFM_ASSUME(Current != Parent);
+
+			// The root will always have a black parent, so this check covers both.
+			if (Parent == Root)
+			{
+				Nodes.SetColor(Parent, EColor::Black);
+				break;
+			}
+			else if (IsBlack(Parent))
+			{
+				break;
+			}
+
+			const NodeIndex GrandParent = Nodes.Parent(Parent);
+
+			UE_AUTORTFM_ASSUME((Parent != GrandParent) && (Current != GrandParent));
+
+			const bool bParentIsLeft = Nodes.Left(GrandParent) == Parent;
+
+			// The uncle is the other node of our parent.
+			const NodeIndex Uncle = bParentIsLeft ? Nodes.Right(GrandParent) : Nodes.Left(GrandParent);
+
+			UE_AUTORTFM_ASSUME((GrandParent != Uncle) && (Parent != Uncle) && (Current != Uncle));
+
+			if (!IsBlack(Uncle))
+			{
+				AUTORTFM_ASSERT(IsBlack(GrandParent));
+				Nodes.SetColor(Parent, EColor::Black);
+				Nodes.SetColor(Uncle, EColor::Black);
+
+				if (GrandParent == Root)
+				{
+					break;
+				}
+				else
+				{
+					Nodes.SetColor(GrandParent, EColor::Red);
+				}
+
+				Current = GrandParent;
+				continue;
+			}
+
+			// Our uncle is black, so we need to swizzle around.
+			const bool bCurrentIsLeft = Nodes.Left(Parent) == Current;
+
+			if (bParentIsLeft)
+			{
+				if (!bCurrentIsLeft)
+				{
+					Nodes.Right(Parent) = Nodes.Left(Current);
+
+					if (InvalidNode != Nodes.Right(Parent))
+					{
+						Nodes.Parent(Nodes.Right(Parent)) = Parent;
+					}
+
+					Nodes.Parent(Parent) = Current;
+					Nodes.Parent(Current) = GrandParent;
+					Nodes.Left(Current) = Parent;
+					Nodes.Left(GrandParent) = Current;
+					std::swap(Parent, Current);
+				}
+
+				Nodes.Parent(Parent) = Nodes.Parent(GrandParent);
+				Nodes.Left(GrandParent) = Nodes.Right(Parent);
+
+				if (InvalidNode != Nodes.Left(GrandParent))
+				{
+					Nodes.Parent(Nodes.Left(GrandParent)) = GrandParent;
+				}
+
+				Nodes.Right(Parent) = GrandParent;
+				Nodes.Parent(GrandParent) = Parent;
+
+				const EColor bGrandParentColor = Nodes.Color(GrandParent);
+				Nodes.SetColor(GrandParent, Nodes.Color(Parent));
+				Nodes.SetColor(Parent, bGrandParentColor);
+			}
+			else
+			{
+				if (bCurrentIsLeft)
+				{
+					Nodes.Left(Parent) = Nodes.Right(Current);
+
+					if (InvalidNode != Nodes.Left(Parent))
+					{
+						Nodes.Parent(Nodes.Left(Parent)) = Parent;
+					}
+
+					Nodes.Parent(Parent) = Current;
+					Nodes.Parent(Current) = GrandParent;
+					Nodes.Right(Current) = Parent;
+					Nodes.Right(GrandParent) = Current;
+					std::swap(Parent, Current);
+				}
+
+				Nodes.Parent(Parent) = Nodes.Parent(GrandParent);
+				Nodes.Right(GrandParent) = Nodes.Left(Parent);
+
+				if (InvalidNode != Nodes.Right(GrandParent))
+				{
+					Nodes.Parent(Nodes.Right(GrandParent)) = GrandParent;
+				}
+
+				Nodes.Left(Parent) = GrandParent;
+				Nodes.Parent(GrandParent) = Parent;
+
+				const EColor bGrandParentColor = Nodes.Color(GrandParent);
+				Nodes.SetColor(GrandParent, Nodes.Color(Parent));
+				Nodes.SetColor(Parent, bGrandParentColor);
+			}
+
+			if (InvalidNode == Nodes.Parent(Parent))
+			{
+				Root = Parent;
+			}
+			else
+			{
+				const NodeIndex GreatGrandParent = Nodes.Parent(Parent);
+
+				if (GrandParent == Nodes.Left(GreatGrandParent))
+				{
+					Nodes.Left(GreatGrandParent) = Parent;
+				}
+				else
+				{
+					Nodes.Right(GreatGrandParent) = Parent;
+				}
+			}
+
+			break;
+		}
+
+		AssertStructureIsOk();
+
+		return true;
+	}
+
+	UE_AUTORTFM_FORCENOINLINE bool Query(const void* const Address, const size_t Size) const
+	{
+		if (AUTORTFM_UNLIKELY(InvalidNode == Root))
+		{
+			return false;
+		}
+
+		const FInterval QueryInterval{Address, Size};
+
+		NodeIndex Current = Root;
+
+		do
+		{
+			const FInterval Interval = Nodes.Interval(Current);
+
+			// This check does not need to prove that QueryInterval is entirely
+			// enclosed within Interval, because if any byte of QueryInterval was in the
+			// original Interval then it **must** already have been new memory.
+			if ((QueryInterval.Start < Interval.End) && (Interval.Start < QueryInterval.End))
+			{
+				return true;
+			}
+
+			const NodeIndex* Next = (QueryInterval.Start < Interval.End) ? Nodes.Ptrs.Left : Nodes.Ptrs.Right;
+			Current = Next[Current];
+		} while (InvalidNode != Current);
+
+		return false;
+	}
+
+	UE_AUTORTFM_FORCEINLINE void AssertStructureIsOk() const
+	{
+		if constexpr (bExtraDebugging)
+		{
+			if (InvalidNode != Root)
+			{
+				AssertNodeIsOk(Root);
+			}
+		}
+	}
+
+	UE_AUTORTFM_FORCENOINLINE void AssertNodeIsOk(NodeIndex Index) const
+	{
+		// We need to use recursion to check because we cannot have any
+		// allocations within the checker itself!
+		if constexpr (bExtraDebugging)
+		{
+			AUTORTFM_ASSERT(InvalidNode != Index);
+			AUTORTFM_ASSERT(Index < Nodes.Count);
+
+			const NodeIndex Parent = Nodes.Parent(Index);
+			const NodeIndex Left = Nodes.Left(Index);
+			const NodeIndex Right = Nodes.Right(Index);
+
+			if (InvalidNode == Parent)
+			{
+				AUTORTFM_ASSERT(Root == Index);
+			}
+			else
+			{
+				AUTORTFM_ASSERT(Nodes.Color(Parent) == EColor::Black || Nodes.Color(Index) == EColor::Black);
+				AUTORTFM_ASSERT((Nodes.Left(Parent) == Index) ^ (Nodes.Right(Parent) == Index));
+			}
+
+			AUTORTFM_ASSERT(Left != Index);
+			AUTORTFM_ASSERT(Right != Index);
+
+			if (InvalidNode != Left)
+			{
+				AssertNodeIsOk(Left);
+			}
+
+			if (InvalidNode != Right)
+			{
+				AssertNodeIsOk(Right);
+			}
+		}
+	}
+};
+
+}  // namespace AutoRTFM
+
+#endif  // (defined(__AUTORTFM) && __AUTORTFM)
diff --git a/Engine/Source/Runtime/AutoRTFM/Private/Sanitizer.cpp b/Engine/Source/Runtime/AutoRTFM/Private/Sanitizer.cpp
--- a/Engine/Source/Runtime/AutoRTFM/Private/Sanitizer.cpp
+++ b/Engine/Source/Runtime/AutoRTFM/Private/Sanitizer.cpp
@@ -258,9 +258,9 @@ void FSanitizer::AbortTransaction(TransactionID ID)
 
 	for (FTransaction* Transaction = CurrentTransaction; Transaction; Transaction = Transaction->Parent)
 	{
-		// TODO(SOL-8992): FIntervalTree::Contains() may incorrectly return
-		// false if the queried interval starts outside of an interval added by
-		// FIntervalTree::Insert() but still overlaps an inserted interval.
+		// TODO(SOL-8992): TIntervalTree::Get() may incorrectly return
+		// null if the queried interval starts outside of an interval added by
+		// TIntervalTree::Insert() but still overlaps an inserted interval.
 		// Fix this.
 		if (FCallstackTree::FHandle const* const ClosedWriteCallstack = Transaction->ClosedWrites.Get(Ptr, Size);
 			AUTORTFM_UNLIKELY(ClosedWriteCallstack))
diff --git a/Engine/Source/Runtime/AutoRTFM/Private/Transaction.h b/Engine/Source/Runtime/AutoRTFM/Private/Transaction.h
--- a/Engine/Source/Runtime/AutoRTFM/Private/Transaction.h
+++ b/Engine/Source/Runtime/AutoRTFM/Private/Transaction.h
@@ -6,12 +6,11 @@
 
 #include "AutoRTFM.h"
 #include "HitSet.h"
-#include "IntervalTree.h"
 #include "IntrusivePool.h"
+#include "MemoryTracker.h"
 #include "Stack.h"
 #include "StackRange.h"
 #include "TaskArray.h"
-#include "Utils.h"
 #include "WriteLog.h"
 
 namespace AutoRTFM
@@ -132,7 +131,7 @@ private:
 	TStack<const void*, 1> DeferredPopAllOnAbortHandlers;
 
 	FHitSet HitSet;
-	FIntervalTree NewMemoryTracker;
+	FMemoryTracker NewMemoryTracker;
 	FWriteLog WriteLog;
 	FStackRange StackRange;
 	TransactionID ID = 0;
