[AutoRTFM] AUTORTFM_DISABLE TTaskArray

Fixes build failures in LTO / PGO builds.

#Synced-CL 51584942

diff --git a/Engine/Source/Runtime/AutoRTFM/Private/TaskArray.h b/Engine/Source/Runtime/AutoRTFM/Private/TaskArray.h
--- a/Engine/Source/Runtime/AutoRTFM/Private/TaskArray.h
+++ b/Engine/Source/Runtime/AutoRTFM/Private/TaskArray.h
@@ -24,7 +24,7 @@ namespace AutoRTFM
 // Key             - the data type used as a key for push / pop.
 // InlineTaskCount - the number of tasks that can be allocated before a performing heap allocations.
 template <typename TaskType, typename KeyType = const void*>
-class TTaskArray
+class AUTORTFM_DISABLE TTaskArray
 {
 	struct FEntry
 	{
