[AutoRTFM] Add autortfm_debug_dump_ir() debugging utility

Add a new autortfm_debug_dump_ir() function that can be called in any function body to have the AutoRTFM compiler dump that function's LLVM IR at each pipeline stage (pre/post EarlyPass and pre/post InstrumentationPass).
The call is automatically stripped after the final dump. Has no effect when AutoRTFM is disabled.

This is useful for inspecting how the AutoRTFM compiler transforms a specific function, without having to dump the entire module.

Also fix missing includes in Diagnostics.h (Function.h, Instruction.h) that caused build failures when the header was included without the full type definitions being available.

Rebuilt Win64 and Linux toolchain binaries.

#Synced-CL 52002210

#Synced-CL 51896064

diff --git a/Engine/Source/Runtime/AutoRTFM/Public/AutoRTFM/CAPI.h b/Engine/Source/Runtime/AutoRTFM/Public/AutoRTFM/CAPI.h
--- a/Engine/Source/Runtime/AutoRTFM/Public/AutoRTFM/CAPI.h
+++ b/Engine/Source/Runtime/AutoRTFM/Public/AutoRTFM/CAPI.h
@@ -211,6 +211,16 @@ UE_AUTORTFM_CRITICAL_INLINE bool autortfm_is_closed(void) AUTORTFM_NOEXCEPT
 }
 #endif
 
+// Debug utility: place a call to this function in a function body to have the
+// AutoRTFM compiler dump that function's IR at each pipeline stage (pre/post
+// EarlyPass, pre/post InstrumentationPass). The call is stripped after the
+// final dump. Has no effect when AutoRTFM is disabled.
+#if UE_AUTORTFM_ENABLED
+	UE_AUTORTFM_API void autortfm_debug_dump_ir(void) AUTORTFM_NOEXCEPT;
+#else
+UE_AUTORTFM_CRITICAL_INLINE void autortfm_debug_dump_ir(void) AUTORTFM_NOEXCEPT {}
+#endif
+
 #if UE_AUTORTFM_ENABLED
 	// Returns the current AutoRTFM context status.
 	// This function is handled specially in the compiler and will be constant
