AutoRTFM: Pre-optimization disabled-function validation and template mode inference

Moves disabled-function detection from the post-optimization InstrumentationPass to a new pre-optimization EarlyPass validation (`ValidateDisabledCalls`). This makes enable->disable violation errors consistent across optimization levels -- they no longer depend on inlining decisions.

## Clang frontend (`AutoRTFM.cpp`)

- Refactored `GetMode()` into priority-ordered helpers: `GetExplicitNonAutoRTFMAttrMode`, `GetExplicitAutoRTFMAttrMode`, `ShouldDefaultToInfer`.
- Legacy attributes (`UE_AUTORTFM_ALWAYS_OPEN`, `UE_AUTORTFM_NOAUTORTFM`, etc.) on a function take priority over `AutoRTFMAttr` on enclosing declarations, but do NOT propagate to nested lambdas/local functions.
- `AutoRTFMAttr` (`AUTORTFM_DISABLE`, `AUTORTFM_ENABLE`, `AUTORTFM_OPEN`, etc.) propagates to nested lambdas and local functions via `OuterOf`.
- Unannotated functions in template contexts (function template specializations, class template members, local structs inside template bodies, `std::` namespace) now default to `infer` mode -- their mode adapts to the types they're instantiated with. Excludes `dllimport` and `extern template` declarations.
- `SetLLVMFunctionAttributesForDefinition` now re-evaluates `GetMode()` when the definition is available, fixing cases where template bodies weren't instantiated at declaration time.

## EarlyPass (`Passes.cpp`)

- `ValidateDisabledCalls()`: Pre-optimization check for enable->disable violations. Uses `ComputeOpenOnlyBlocks()` for `autortfm_is_closed()` flow-control analysis to avoid false positives on guarded calls. Recognizes `always_inline` functions containing `IsClosed()` guards (e.g. `UnreachableIfClosed`), including when the guard call is an `invoke` (exceptions enabled). Excludes global constructors and destructors (`llvm.global_ctors` / `llvm.global_dtors`).
- Added `-autortfm-no-validate-disabled-calls` cl::opt to disable the compile-time validation.

## InstrumentationPass (`Passes.cpp`)

- Removed disabled-function error from `ClosedInstrumentationVisitor` and `ProcessDeferredCalls` -- now handled by EarlyPass.
- `GatherFunctionsExcludedForInstrumentation()` consolidated with `ValidateDisabledCalls` via shared `GatherGlobalCtorDtorRoots()` helper; now also excludes global destructors.

## Utilities (`Utils.cpp`, `Utils.h`)

- `ComputeOpenOnlyBlocks()`: Flow-control analysis tracking `autortfm_is_closed()` predicates through boolean operations and branches. Used by EarlyPass validation.
- `EmitDisabledFunctionError()`: Simplified -- removed mode metadata lookup and inlining chain annotations. Error chain now shows the specific callee and caller function names.
- `DebugDumpTaggedFunctions()` / `autortfm_debug_dump_ir()`: New diagnostic tool for dumping a function's IR at each pipeline stage.

## Constants (`Constants.h`)

- Added `IsOpen()` and `IsNoSanitize()` constexpr helper functions for `autortfm_mode` classification.

## Documentation

- Updated `ModesAndAttributes.md` mode attribute rules with template inference heuristic.

## Tests

- 249 lit tests (19 new, 21 edited). New tests cover: attribute propagation to lambdas, local structs inside class templates, `extern template` declarations, global destructor exclusion, disabled-function validation with `IsClosed()` guards (including `IsContextStatus` and `UnreachableIfTransactional` variants), guarded calls in inferred functions, inferred disable chain diagnostics, `std::` namespace inference, validation suppression via `-autortfm-no-validate-disabled-calls`, `IsClosed()` guard with exceptions enabled, static locals in enabled lambdas, inlined mode info diagnostics. Edited tests updated to check specific callee/caller names in error diagnostics.

#Synced-CL 52126916

diff --git a/Engine/Source/Runtime/AutoRTFM/Documentation/ModesAndAttributes.md b/Engine/Source/Runtime/AutoRTFM/Documentation/ModesAndAttributes.md
--- a/Engine/Source/Runtime/AutoRTFM/Documentation/ModesAndAttributes.md
+++ b/Engine/Source/Runtime/AutoRTFM/Documentation/ModesAndAttributes.md
@@ -5,7 +5,7 @@ every function is compiled with a single, mutually-exclusive [AutoRTFM mode](#au
 
 The AutoRTFM mode controls the behavior of the [closed variant](README.md#open-and-closed-code) of the function.
 
-AutoRTFM modes are specified using [AutoRTFM attributes](#autortfm-attributes) on lambdas, functions, methods and classes.
+AutoRTFM modes are specified using [AutoRTFM attributes](#mode-attributes) on lambdas, functions, methods and classes.
 
 The UBT module flag `bInjectAutoRTFMAttributeDisables` can also be used to apply [`AUTORTFM_DISABLE`](#autortfm-disabled) to all module exports (e.g. `FOO_API`).
 
@@ -19,7 +19,7 @@ instrumentation has an AutoRTFM mode. This is one of the following:
 AutoRTFM instrumentation **is** generated for the function. This is the default
 mode for targets compiled with AutoRTFM instrumentation enabled.
 
-An AutoRTFM-enabled function is callable from the **open** and **closed**.
+An AutoRTFM-enabled function is callable from the **open** and the **closed**.
 
 Represented by the enum `autortfm_mode_enable`.
 
@@ -30,11 +30,12 @@ This function is callable from the **open**, but not the **closed**.
 
 Attempting to call an AutoRTFM-disabled function from an
 [AutoRTFM-enabled](#autortfm-enabled) function is a compile-time error, unless
-the call is unconditionally preceded with a call to `AutoRTFM::UnreachableIfClosed()`.
+the call is unconditionally preceded with a call to
+`AutoRTFM::UnreachableIfClosed()` or `AutoRTFM::UnreachableIfTransactional()`.
 
-Use of `AutoRTFM::UnreachableIfClosed()` will allow code to compile that
-calls AutoRTFM-disabled functions from [AutoRTFM-enabled](#autortfm-enabled)
-functions, but `AutoRTFM::UnreachableIfClosed()` will trigger a runtime
+Use of `AutoRTFM::UnreachableIfClosed()` / `AutoRTFM::UnreachableIfTransactional()`
+will allow code to compile that calls AutoRTFM-disabled functions from
+[AutoRTFM-enabled](#autortfm-enabled) functions, but will trigger a runtime
 assertion if called from the closed.
 
 Represented by the enum `autortfm_mode_disable`.
@@ -58,9 +59,9 @@ AutoRTFM-inferred function makes a direct call to an
 AutoRTFM-inferred functions.
 
 Similarly to [AutoRTFM-disabled](#autortfm-disabled) functions, use of
-`AutoRTFM::UnreachableIfClosed()` can be used to allow code to compile that
-would otherwise fail, but with the same caveat that this may cause runtime
-assertions.
+`AutoRTFM::UnreachableIfClosed()` or `AutoRTFM::UnreachableIfTransactional()`
+can be used to allow code to compile that would otherwise fail, but with the
+same caveat that this may cause runtime assertions.
 
 Represented by the enum `autortfm_mode_infer`.
 
@@ -100,10 +101,11 @@ The following attributes can be applied to lambdas, functions, classes and struc
 | Attribute                               | Mode |
 |-----------------------------------------|------|
 | `AUTORTFM_ENABLE`                       | Unconditionally applies the [AutoRTFM-enabled](#autortfm-enabled) mode. |
+| `AUTORTFM_ENABLE_IF(<cond>)`            | Conditionally applies the [AutoRTFM-enabled](#autortfm-enabled) mode if `<cond>` evaluates to true, otherwise behaves as if the attribute did not exist. |
 | `AUTORTFM_DISABLE`                      | Unconditionally applies the [AutoRTFM-disabled](#autortfm-disabled) mode. |
 | `AUTORTFM_DISABLE_IF(<cond>)`           | Conditionally applies the [AutoRTFM-disabled](#autortfm-disabled) mode if `<cond>` evaluates to true, otherwise behaves as if the attribute did not exist. |
 | `AUTORTFM_OPEN`                         | Unconditionally applies the [AutoRTFM-open](#autortfm-open) mode. |
-| `AUTORTFM_OPEN_NO_SANITIZE`             | Unconditionally applies the [AutoRTFM-open-with-no-sanitize](#autortfm-open) mode. |
+| `AUTORTFM_OPEN_NO_SANITIZE`             | Unconditionally applies the [AutoRTFM-open-with-no-sanitize](#autortfm-open-with-no-sanitize) mode. |
 | `AUTORTFM_INFER`                        | Unconditionally applies the [AutoRTFM-inferred](#autortfm-inferred) mode. |
 
 Attributes on functions and methods need to be placed on the declaration (usually in the `.h`) and positioned before the return type (C-declaration style).\
@@ -150,8 +152,8 @@ Returns the [AutoRTFM mode](#autortfm-modes) of a function by call expression, w
 ```C++
 AUTORTFM_ENABLE void MyOverload(int);
 AUTORTFM_DISABLE void MyOverload(void*);
-static_assert(AUTORTFM_MODE_OF(MyOverload(42)) == autortfm_mode_enable);
-static_assert(AUTORTFM_MODE_OF(MyOverload(nullptr)) == autortfm_mode_disable);
+static_assert(AUTORTFM_MODE_OF_CALL(MyOverload(42)) == autortfm_mode_enable);
+static_assert(AUTORTFM_MODE_OF_CALL(MyOverload(nullptr)) == autortfm_mode_disable);
 ```
 
 ### `AutoRTFM::CallMode<FunctorType, ArgTypes...>`
@@ -178,7 +180,7 @@ of `Type`.
 ```C++
 struct DisabledDestructor
 {
-    AUTORTFM_DISABLE ~ExplicitDestructorDisabled();
+    AUTORTFM_DISABLE ~DisabledDestructor();
 };
 static_assert(AutoRTFM::DestructorMode<DisabledDestructor> == autortfm_mode_disable);
 ```
@@ -191,21 +193,24 @@ of `Type` when called with `ArgTypes`.
 ```C++
 struct OverloadedConstructors
 {
-    AUTORTFM_DISABLE ExplicitConstructors();
-    AUTORTFM_ENABLE  ExplicitConstructors(int);
+    AUTORTFM_DISABLE OverloadedConstructors();
+    AUTORTFM_ENABLE  OverloadedConstructors(int);
 };
-static_assert(AutoRTFM::ConstructorMode<ExplicitConstructors> == autortfm_mode_disable);
-static_assert(AutoRTFM::ConstructorMode<ExplicitConstructors, int> == autortfm_mode_enable);
+static_assert(AutoRTFM::ConstructorMode<OverloadedConstructors> == autortfm_mode_disable);
+static_assert(AutoRTFM::ConstructorMode<OverloadedConstructors, int> == autortfm_mode_enable);
 ```
 
 ## Mode Attribute Rules
 
 The mode used is determined by the following rules, starting with the highest precedence:
 
-1. A lambda, function, class or struct directly annotated with any of the
+1. Legacy attributes (`UE_AUTORTFM_ALWAYS_OPEN`, `UE_AUTORTFM_ALWAYS_OPEN_NO_MEMORY_VALIDATION`,
+   `UE_AUTORTFM_NOAUTORTFM`) on the function itself take highest priority.
+   These do **not** propagate to nested lambdas or local functions.
+2. A lambda, function, class or struct directly annotated with any of the
    unconditional or conditionally-true [AutoRTFM mode attributes](#mode-attributes)
    will use that attribute's mode.
-2. A "naked" (passed directly, not via a variable) lambda, with no explicit
+3. A "naked" (passed directly, not via a variable) lambda, with no explicit
    AutoRTFM mode attribute, passed to a function parameter annotated with the
    `AUTORTFM_IMPLICIT_ENABLE` or `AUTORTFM_IMPLICIT_DISABLE` attributes will be
    [enabled](#autortfm-enabled) or [disabled](#autortfm-disabled), respectively. \
@@ -213,20 +218,35 @@ The mode used is determined by the following rules, starting with the highest pr
    not explicitly annotated, and the single `AutoRTFM::Transact()` parameter
    is annotated with `AUTORTFM_IMPLICIT_ENABLE`, so the lambda will be [enabled](#autortfm-enabled),
    regardless of the outer function or class attributes.
-3. A function scoped lambda, or function scoped class will use the AutoRTFM mode
+4. A function scoped lambda, or function scoped class will use the AutoRTFM mode
    of the parent function.
-4. An unannotated (or annotated with an conditionally-false [AutoRTFM mode attribute](#mode-attributes))
-   `class` or `struct` method will use the mode of the parent `class` or `struct`.
 5. An unannotated (or annotated with an conditionally-false [AutoRTFM mode attribute](#mode-attributes))
+   `class` or `struct` method will use the mode of the parent `class` or `struct`.
+6. An unannotated (or annotated with an conditionally-false [AutoRTFM mode attribute](#mode-attributes))
    nested `class` or `struct` will use the mode of the parent `class` or `struct`.
-6. An unnested `class` or `struct` will be [AutoRTFM-enabled](#autortfm-enabled)
+7. An unannotated function in a **template context** — i.e. a function template
+   specialization, a member of a class template specialization, a function in
+   the `std::` namespace, or a local function/struct defined inside any of these
+   — will default to [AutoRTFM-inferred](#autortfm-inferred) if the function has
+   a visible body. This allows template instantiations to adapt their mode to the
+   types they are instantiated with. For example, `std::vector<DisabledType>::push_back()`
+   will infer to [AutoRTFM-disabled](#autortfm-disabled) because it calls the
+   disabled copy/move constructor. \
+   Note: Functions with explicit AutoRTFM attributes (rules 1–2) are never
+   affected by this heuristic. AutoRTFM runtime functions that intentionally
+   accept disabled functors (e.g. `AutoRTFM::Open`, `AutoRTFM::OnCommit`) must
+   be explicitly annotated with `AUTORTFM_ENABLE` to prevent inference. \
+   Note: `extern template` declarations and `dllimport` functions are excluded
+   from this heuristic because their bodies are not available for analysis.
+8. An unnested `class` or `struct` will be [AutoRTFM-enabled](#autortfm-enabled)
    unless the module has instrumentation disabled.
 
 Perhaps a simpler way to describe these rules is:
 > With exception of the `AUTORTFM_IMPLICIT_ENABLE` and `AUTORTFM_IMPLICIT_DISABLE`
 parameter attributes, **the AutoRTFM mode used by a function is the mode taken
 from the closest enclosing scope with an unconditional or conditionally-true
-AutoRTFM attribute**.
+AutoRTFM attribute**. If no such attribute exists and the function is in a
+template context, it defaults to [AutoRTFM-inferred](#autortfm-inferred).
 
 ### Examples
 
@@ -254,24 +274,28 @@ void DisabledStruct::DisabledByAttributeOnStruct()
     // Lambda that is AutoRTFM-disabled by the AUTORTFM_DISABLE attribute on DisabledStruct.
     auto LambdaDisabledByAttributeOnStruct = []{ /* ... */ };
 
-    // Lambda that is AutoRTFM-enabled by the explicit AUTORTFM_OPEN attribute on the lambda.
+    // Lambda that is AutoRTFM-open by the explicit AUTORTFM_OPEN attribute on the lambda.
     auto ExplicitOpenLambda = [] AUTORTFM_OPEN { /* ... */ };
+}
 
+void DisabledStruct::EnabledByExplicitAttribute()
+{
     if (SomeConditional())
     {
         // UnreachableIfClosed() will runtime-error if this branch is taken in the closed.
         AutoRTFM::UnreachableIfClosed("I promise that this is not reachable in the closed");
-        // The UnreachableIfClosed() above prevents this AutoRTFM-enabled function call
-        // from an AutoRTFM-disabled function from being a compile-time error.
-        EnabledByExplicitAttribute();
+        // The UnreachableIfClosed() above prevents this call to an AutoRTFM-disabled
+        // function from an AutoRTFM-enabled function from being a compile-time error.
+        DisabledByAttributeOnStruct();
     }
 
-    // Attempting to call an AutoRTFM-enabled function from this AutoRTFM-disabled
-    // function, without an unconditionally preceded call to AutoRTFM::UnreachableIfClosed()
-    // is a compile-time error:
-    //   error: Cannot call AutoRTFM-disabled function 'EnabledByExplicitAttribute()' 
-    //          from the AutoRTFM-enabled function 'DisabledByAttributeOnStruct()'.
-    EnabledByExplicitAttribute();
+    // Attempting to call an AutoRTFM-disabled function from this AutoRTFM-enabled
+    // function, without a preceding call to AutoRTFM::UnreachableIfClosed(), is a
+    // compile-time error:
+    //   error: Cannot call AutoRTFM-disabled function from an AutoRTFM-enabled function.
+    //   note: AutoRTFM-disabled function 'DisabledStruct::DisabledByAttributeOnStruct()'
+    //   note: called by AutoRTFM-enabled function 'DisabledStruct::EnabledByExplicitAttribute()'
+    DisabledByAttributeOnStruct();
 }
 ```
 
@@ -307,6 +331,10 @@ that is making a call to either:
 - An [AutoRTFM-inferred](#autortfm-inferred) function which transitively calls
   an [AutoRTFM-disabled](#autortfm-disabled) function.
 
+When the callee was inferred to disable, the diagnostic will show the full
+chain of calls that caused the inference, making it easier to identify the
+root cause.
+
 The compiler is erroring because the callee function has no closed variant for
 the closed caller function to call.\
 To fix this, you have a few choices:
@@ -316,12 +344,13 @@ To fix this, you have a few choices:
    this is a method and the entire class should not be callable from the closed).
    This option is ideal for reducing code size bloat, improving compilation
    performance, and alerting developers to AutoRTFM mistakes, but
-   `AUTORTFM_DISABLE` is "viral" and may cause callers of the callee function to
+   `AUTORTFM_DISABLE` is "viral" and may cause callers of the disabled function to
    start erroring with a similar message until all of those are also fixed up.\
    Use of `AUTORTFM_DISABLE` for these errors is the preferred long-term solution.
-2. Alternatively, you may wish to use `AutoRTFM::UnreachableIfClosed()` at the
-   top of the callee function. This will silence the compiler error, but if
-   the callee function is itself called by a closed transaction, then the
+2. Alternatively, you may wish to use `AutoRTFM::UnreachableIfClosed()` or
+   `AutoRTFM::UnreachableIfTransactional()` in the caller function, before the
+   call to the disabled callee. This will silence the compiler error, but if
+   the caller function is itself called by a closed transaction, then the
    runtime will assert and crash.\
    This can be used as a short-term solution until `AUTORTFM_DISABLE` can be
    correctly applied.
diff --git a/Engine/Source/Runtime/AutoRTFM/Public/AutoRTFM/Constants.h b/Engine/Source/Runtime/AutoRTFM/Public/AutoRTFM/Constants.h
--- a/Engine/Source/Runtime/AutoRTFM/Public/AutoRTFM/Constants.h
+++ b/Engine/Source/Runtime/AutoRTFM/Public/AutoRTFM/Constants.h
@@ -131,6 +131,18 @@ constexpr inline bool IsDisabled(autortfm_mode Mode)
 	return Mode == autortfm_mode_disable || Mode == autortfm_mode_disable_no_sanitize || Mode == autortfm_mode_no_autortfm;
 }
 
+// Returns true if mode is autortfm_mode_open or autortfm_mode_open_no_sanitize
+constexpr inline bool IsOpen(autortfm_mode Mode)
+{
+	return Mode == autortfm_mode_open || Mode == autortfm_mode_open_no_sanitize;
+}
+
+// Returns true if mode suspends the AutoRTFM sanitizer (open_no_sanitize or disable_no_sanitize)
+constexpr inline bool IsNoSanitize(autortfm_mode Mode)
+{
+	return Mode == autortfm_mode_open_no_sanitize || Mode == autortfm_mode_disable_no_sanitize;
+}
+
 }
 
 namespace AutoRTFM::Constants
