[AutoRTFM] Default sanitizer to recording closed callstacks

It's slower, but so much more helpful.

Also - don't crash if changing sanitizer modes via command line flags before AutoRTFM is initialized.

#Synced-CL 52392476

diff --git a/Engine/Source/Runtime/AutoRTFM/Private/Sanitizer.cpp b/Engine/Source/Runtime/AutoRTFM/Private/Sanitizer.cpp
--- a/Engine/Source/Runtime/AutoRTFM/Private/Sanitizer.cpp
+++ b/Engine/Source/Runtime/AutoRTFM/Private/Sanitizer.cpp
@@ -131,7 +131,7 @@ private:
 FSanitizer* FSanitizer::Instance = nullptr;
 std::atomic<Sanitizer::EMode> FSanitizer::Mode{EMode::Error};
 std::atomic<size_t> FSanitizer::NumTransactionsRunning{0};
-std::atomic<bool> FSanitizer::bRecordClosedWriteCallstacks{false};
+std::atomic<bool> FSanitizer::bRecordClosedWriteCallstacks{true};
 
 struct AUTORTFM_DISABLE_NO_SANITIZE FSanitizer::FLock
 {
@@ -519,7 +519,10 @@ void SetMode(EMode NewMode)
 	if (FSanitizer::Mode != NewMode)
 	{
 		FSanitizer::Mode = NewMode;
-		GExternAPI.OnSanitizerModeChanged();
+		if (GExternAPI.OnSanitizerModeChanged)
+		{
+			GExternAPI.OnSanitizerModeChanged();
+		}
 	}
 }
 
@@ -533,7 +536,10 @@ void SetRecordClosedWriteCallstacks(bool bEnable)
 	if (FSanitizer::bRecordClosedWriteCallstacks != bEnable)
 	{
 		FSanitizer::bRecordClosedWriteCallstacks = bEnable;
-		GExternAPI.OnSanitizerRecordClosedWriteCallstacksChanged();
+		if (GExternAPI.OnSanitizerRecordClosedWriteCallstacksChanged)
+		{
+			GExternAPI.OnSanitizerRecordClosedWriteCallstacksChanged();
+		}
 	}
 }
 
