[AutoRTFM] Fix C++23 build errors.

std::aligned_storage_t is marked as deprecated in C++23, which breaks the build when warnings are treated as errors.

#Synced-CL 52379757

diff --git a/Engine/Source/Runtime/AutoRTFM/Public/AutoRTFM.h b/Engine/Source/Runtime/AutoRTFM/Public/AutoRTFM.h
--- a/Engine/Source/Runtime/AutoRTFM/Public/AutoRTFM.h
+++ b/Engine/Source/Runtime/AutoRTFM/Public/AutoRTFM.h
@@ -19,9 +19,8 @@
 
 #include "AutoRTFM/Task.h"
 
-#include <cstdarg>
+#include <cstddef>
 #include <new>
-#include <stdlib.h>
 #include <type_traits>
 #include <utility>
 
@@ -614,7 +613,7 @@ UE_AUTORTFM_CRITICAL_INLINE ReturnType Open(const FunctorType& Functor)
 
 		if constexpr (Mode == EReturnFromOpenMode::CopyConstructInClosed || Mode == EReturnFromOpenMode::MoveConstructInClosed)
 		{
-			std::aligned_storage_t<sizeof(ReturnType), alignof(ReturnType)> OpenReturnValueMemory;
+			alignas(ReturnType) std::byte OpenReturnValueMemory[sizeof(ReturnType)];
 
 			// Call Functor in the open, and in-place construct the return value into OpenReturnValue.
 			struct FArgs
@@ -622,7 +621,7 @@ UE_AUTORTFM_CRITICAL_INLINE ReturnType Open(const FunctorType& Functor)
 				const FunctorType& Functor;
 				void* OpenReturnValueMemory;
 			};
-			FArgs FunctorAndReturnValue{Functor, &OpenReturnValueMemory};
+			FArgs FunctorAndReturnValue{Functor, OpenReturnValueMemory};
 			{
 				auto Call = [] AUTORTFM_DISABLE(void* ArgPtr)
 				{
@@ -633,7 +632,7 @@ UE_AUTORTFM_CRITICAL_INLINE ReturnType Open(const FunctorType& Functor)
 			}
 
 			// Re-interpret the OpenReturnValueMemory as a ReturnType&.
-			ReturnType& OpenReturnValue = *std::launder(reinterpret_cast<ReturnType*>(&OpenReturnValueMemory));
+			ReturnType& OpenReturnValue = *std::launder(reinterpret_cast<ReturnType*>(OpenReturnValueMemory));
 
 			if constexpr (Mode == EReturnFromOpenMode::CopyConstructInClosed)
 			{
