Fix typo in AutoRTFM hazard report.

The correct class name is FTransactionallySafeRWLock, not FTransactionallySafeFRWLock.

#Synced-CL 52379757

diff --git a/Engine/Source/Runtime/AutoRTFM/Private/ReportHazard.cpp b/Engine/Source/Runtime/AutoRTFM/Private/ReportHazard.cpp
--- a/Engine/Source/Runtime/AutoRTFM/Private/ReportHazard.cpp
+++ b/Engine/Source/Runtime/AutoRTFM/Private/ReportHazard.cpp
@@ -37,7 +37,7 @@ A read-write lock was accessed while inside an AutoRTFM transaction. AutoRTFM do
 read-write locks, but it does support transactionally-safe read-write locks that can be used instead.
 
 Use the call stack below to find the FRWLock accessed immediately before this error occurred, and
-replace it with FTransactionallySafeFRWLock. If FReadScopeLock, FWriteScopeLock or FRWScopeLock are
+replace it with FTransactionallySafeRWLock. If FReadScopeLock, FWriteScopeLock or FRWScopeLock are
 used, replace them with UE::TReadScopeLock, UE::TWriteScopeLock or UE::TRWScopeLock.
 )";
 			break;
