[AutoRTFM] Improve a comment on an enumerator value

#Synced-CL 52547192

diff --git a/Engine/Source/Runtime/AutoRTFM/Public/AutoRTFM/Constants.h b/Engine/Source/Runtime/AutoRTFM/Public/AutoRTFM/Constants.h
--- a/Engine/Source/Runtime/AutoRTFM/Public/AutoRTFM/Constants.h
+++ b/Engine/Source/Runtime/AutoRTFM/Public/AutoRTFM/Constants.h
@@ -80,7 +80,11 @@ extern "C"
 		// Do not generate AutoRTFM instrumentation for the function.
 		// Calls to the function from AutoRTFM-enabled functions will directly call
 		// the open function without any pre / post callbacks or trampolines.
-		// AutoRTFM sanitizer will ignore these functions.
+		// The AutoRTFM sanitizer ignore these functions - no calls to
+		// autortfm_disable_sanitizer() or autortfm_enable_sanitizer() will be
+		// emitted, nor will any calls to autortfm_sanitizer_open_write() be
+		// emitted for writes, however calls to any open functions made by an
+		// internal function may call autortfm_sanitizer_open_write().
 		autortfm_mode_internal = 6,
 
 		// [[clang::noautortfm]]
