[AutoRTFM] Allow inlining of functions annotated with autortfm modes

Previously, the EarlyPass inserted autortfm_called_no_autortfm() / autortfm_disable_called() directly into disabled function bodies, and the InstrumentationPass handled open/internal functions entirely through trampolines.
Functions with non-default modes could not be inlined across mode boundaries because there was no IR-level representation of mode regions.

This change introduces begin_mode/end_mode call-site markers that encode mode regions into the IR.
The markers survive inlining, allowing the optimizer to inline functions across mode boundaries while preserving correct per-instruction mode information for the InstrumentationPass.

Key changes:
* EmitModeMarkers() wraps call sites to non-default-mode callees with begin_mode/end_mode pairs.
Exception cleanup pads ensure markers are properly balanced on unwind paths (both WinEH and Itanium).
* EmitSanitizerCalls() inserts direct disable_sanitizer/enable_sanitizer calls around IsNoSanitize() function bodies, with elision when the body has no sanitizer-enabled callees.
* ConvertModeMarkers() in the InstrumentationPass converts markers to runtime calls (maybe_throw for open boundaries in closed context, CalledNoAutoRTFM for no_autortfm in closed context) or strips them.
* ModeStackWalker walks begin_mode/end_mode pairs in RPO to compute the mode stack at each instruction, used by ComputeInstructionModes, ConvertModeMarkers, and the sanitizer/instrumentation visitors.
* Musttail calls use an operand bundle to defer end_mode across the musttail+ret boundary.

Now that all AutoRTFM modes do not prevent inlining, replace the AutoRTFM-internal classes with AUTORTFM_INTERNAL instead of AUTORTFM_DISABLE_NO_SANITIZE.

#Synced-CL 52677305

diff --git a/Engine/Source/Programs/AutoRTFMTests/Private/ExceptionTests.cpp b/Engine/Source/Programs/AutoRTFMTests/Private/ExceptionTests.cpp
--- a/Engine/Source/Programs/AutoRTFMTests/Private/ExceptionTests.cpp
+++ b/Engine/Source/Programs/AutoRTFMTests/Private/ExceptionTests.cpp
@@ -218,4 +218,220 @@ TEST_CASE("Exceptions")
 	}
 }
 
+// ---------------------------------------------------------------------------
+// Mode-observation + abort-revert tests for destructors during exception unwind.
+//
+// Each test creates a function in a specific mode that:
+// 1. Constructs an RAII object whose destructor writes to memory and records
+//    the observed mode (IsClosed).
+// 2. Throws an exception.
+//
+// The caller wraps this in try/catch inside a transaction, then aborts.
+// We verify:
+// - The destructor ran (the write happened).
+// - The destructor observed the correct mode (IsClosed).
+// - For closed-mode destructors: the write IS reverted on abort.
+// - For open-mode destructors: the write is NOT reverted on abort.
+// ---------------------------------------------------------------------------
+
+namespace
+{
+
+struct FModeRecorder
+{
+	int* WriteTarget;
+	bool* ObservedIsClosed;
+
+	~FModeRecorder()
+	{
+		if (WriteTarget)
+		{
+			*WriteTarget = 99;
+		}
+		if (ObservedIsClosed)
+		{
+			*ObservedIsClosed = AutoRTFM::IsClosed();
+		}
+	}
+};
+
+// Each function constructs an FModeRecorder locally so the destructor
+// runs inside the function's mode during exception unwinding.
+
+AUTORTFM_OPEN static void OpenThenThrow(int* WriteTarget, bool* ObservedIsClosed)
+{
+	FModeRecorder Recorder{WriteTarget, ObservedIsClosed};
+	throw 42;
+}
+
+AUTORTFM_OPEN_NO_VALIDATION static void OpenNoSanitizeThenThrow(int* WriteTarget, bool* ObservedIsClosed)
+{
+	FModeRecorder Recorder{WriteTarget, ObservedIsClosed};
+	throw 42;
+}
+
+[[clang::autortfm(autortfm_mode_internal)]] static void InternalThenThrow(int* WriteTarget, bool* ObservedIsClosed)
+{
+	FModeRecorder Recorder{WriteTarget, ObservedIsClosed};
+	throw 42;
+}
+
+// Disable-mode throwing function. Called via OpenCallDisableThenThrow
+// to avoid the compile error of calling a disabled function from enabled code.
+AUTORTFM_DISABLE static void DisableThenThrow(int* WriteTarget, bool* ObservedIsClosed)
+{
+	FModeRecorder Recorder{WriteTarget, ObservedIsClosed};
+	throw 42;
+}
+
+AUTORTFM_OPEN static void OpenCallDisableThenThrow(int* WriteTarget, bool* ObservedIsClosed)
+{
+	DisableThenThrow(WriteTarget, ObservedIsClosed);
+}
+
+UE_AUTORTFM_ALWAYS_OPEN static void AssignIntPointer(int* Pointer, int Value)
+{
+	*Pointer = Value;
+}
+
+} // anonymous namespace
+
+TEST_CASE("Exceptions.DestructorModes")
+{
+	AutoRTFMTestUtils::FScopedRetry Retry(AutoRTFM::ForTheRuntime::EAutoRTFMRetryTransactionState::NoRetry);
+
+	SECTION("Destructor in enable mode: closed=true")
+	{
+		int WriteTarget = 0;
+		bool ObservedIsClosed = false;
+
+		AutoRTFM::Transact([&]
+		{
+			try
+			{
+				// FModeRecorder constructed inline in the enable-mode lambda.
+				// Its destructor runs in enable mode (closed) during unwind.
+				FModeRecorder Recorder{&WriteTarget, &ObservedIsClosed};
+				throw 42;
+			}
+			catch (...) {}
+		});
+
+		REQUIRE(ObservedIsClosed == true);
+		// The destructor wrote 99 in closed code, and the transaction
+		// committed, so the write is visible.
+		REQUIRE(WriteTarget == 99);
+	}
+
+	SECTION("Destructor in open mode: closed=false, write NOT reverted on abort")
+	{
+		int WriteTarget = 0;
+		bool ObservedIsClosed = false;
+
+		AutoRTFM::Testing::Abort([&]
+		{
+			try
+			{
+				// FModeRecorder constructed inside OpenThenThrow (open mode).
+				// Its destructor runs in open mode during unwind.
+				OpenThenThrow(&WriteTarget, &ObservedIsClosed);
+			}
+			catch (...) {}
+			AutoRTFM::AbortTransaction();
+		});
+
+		REQUIRE(ObservedIsClosed == false);
+		REQUIRE(WriteTarget == 99); // NOT reverted (open write).
+	}
+
+	SECTION("Destructor in open_no_sanitize mode: closed=false, write NOT reverted on abort")
+	{
+		int WriteTarget = 0;
+		bool ObservedIsClosed = false;
+
+		AutoRTFM::Testing::Abort([&]
+		{
+			try
+			{
+				OpenNoSanitizeThenThrow(&WriteTarget, &ObservedIsClosed);
+			}
+			catch (...) {}
+			AutoRTFM::AbortTransaction();
+		});
+
+		REQUIRE(ObservedIsClosed == false);
+		REQUIRE(WriteTarget == 99); // NOT reverted (open write).
+	}
+
+	SECTION("Destructor in internal mode: closed=false, write NOT reverted on abort")
+	{
+		int WriteTarget = 0;
+		bool ObservedIsClosed = false;
+
+		AutoRTFM::Testing::Abort([&]
+		{
+			try
+			{
+				InternalThenThrow(&WriteTarget, &ObservedIsClosed);
+			}
+			catch (...) {}
+			AutoRTFM::AbortTransaction();
+		});
+
+		REQUIRE(ObservedIsClosed == false);
+		REQUIRE(WriteTarget == 99); // NOT reverted (internal = open write).
+	}
+
+	SECTION("Destructor in disable mode: closed=false, write NOT reverted on abort")
+	{
+		int WriteTarget = 0;
+		bool ObservedIsClosed = false;
+
+		AutoRTFM::Testing::Abort([&]
+		{
+			try
+			{
+				// Call via an open wrapper to avoid compile error
+				// (disabled functions can't be called from enabled code).
+				OpenCallDisableThenThrow(&WriteTarget, &ObservedIsClosed);
+			}
+			catch (...) {}
+			AutoRTFM::AbortTransaction();
+		});
+
+		REQUIRE(ObservedIsClosed == false);
+		REQUIRE(WriteTarget == 99); // NOT reverted (disable = open write).
+	}
+}
+
+#if AUTORTFM_SANITIZER
+TEST_CASE("Exceptions.Sanitizer")
+{
+	AutoRTFMTestUtils::FScopedRetry Retry(AutoRTFM::ForTheRuntime::EAutoRTFMRetryTransactionState::NoRetry);
+	AUTORTFM_SCOPED_SANITIZER_MODE_WARN();
+
+	SECTION("Sanitizer re-enabled after exception through open_no_sanitize")
+	{
+		AutoRTFMTestUtils::FCaptureWarningContext WarningContext;
+		int I = 0;
+
+		AutoRTFM::Transact([&]
+		{
+			I = 42; // Closed write.
+			try
+			{
+				OpenNoSanitizeThenThrow(nullptr, nullptr);
+			}
+			catch (...) {}
+
+			// After the catch, the sanitizer MUST be re-enabled.
+			// This open write conflicts with I = 42, so it should warn.
+			AssignIntPointer(&I, 24);
+		});
+
+		REQUIRE(WarningContext.HasWarningSubstring(AutoRTFMTestUtils::kMemoryModifiedInOpenWarning));
+	}
+}
+#endif // AUTORTFM_SANITIZER
+
 #endif // AUTORTFM_EXCEPTIONS_ENABLED
diff --git a/Engine/Source/Runtime/AutoRTFM/Private/API.cpp b/Engine/Source/Runtime/AutoRTFM/Private/API.cpp
--- a/Engine/Source/Runtime/AutoRTFM/Private/API.cpp
+++ b/Engine/Source/Runtime/AutoRTFM/Private/API.cpp
@@ -399,35 +399,25 @@ AUTORTFM_INTERNAL autortfm_transaction_status RTFM_autortfm_close(
 	return autortfm_transaction_executing;
 }
 
-extern "C" UE_AUTORTFM_NOAUTORTFM UE_AUTORTFM_API void autortfm_pre_open(bool bNoSanitize)
-{
 #if AUTORTFM_SANITIZER
-	if (bNoSanitize)
-	{
-		Sanitizer::BeginDisable();
-	}
-#endif  // AUTORTFM_SANITIZER
+extern "C" UE_AUTORTFM_NOAUTORTFM UE_AUTORTFM_API Sanitizer::FDisableHandle autortfm_disable_sanitizer() noexcept
+{
+	return Sanitizer::BeginDisable();
 }
 
-extern "C" UE_AUTORTFM_NOAUTORTFM UE_AUTORTFM_API void autortfm_post_open(bool bNoSanitize)
+extern "C" UE_AUTORTFM_NOAUTORTFM UE_AUTORTFM_API void autortfm_enable_sanitizer(Sanitizer::FDisableHandle Handle) noexcept
 {
-	FContext* const Context = FContext::Get();
+	Sanitizer::EndDisable(Handle);
+}
+#endif  // AUTORTFM_SANITIZER
 
+extern "C" UE_AUTORTFM_NOAUTORTFM UE_AUTORTFM_API void autortfm_maybe_throw() noexcept
+{
+	FContext* const Context = FContext::Get();
 	Context->MaybeThrow();
-
-#if AUTORTFM_SANITIZER
-	if (bNoSanitize)
-	{
-		// Note: This must come *after* MaybeThrow(), because we'll only throw
-		// if the current transaction at autortfm_pre_open() was aborted and
-		// popped, in which case the current transaction at autortfm_post_open()
-		// will now be an ancestor transaction.
-		Sanitizer::EndDisable();
-	}
-#endif  // AUTORTFM_SANITIZER
 }
 
-extern "C" UE_AUTORTFM_NOAUTORTFM UE_AUTORTFM_API void autortfm_pre_static_local_initializer()
+extern "C" UE_AUTORTFM_NOAUTORTFM UE_AUTORTFM_API void autortfm_pre_static_local_initializer() noexcept
 {
 	if (FContext* Context = FContext::Get())
 	{
@@ -435,7 +425,7 @@ extern "C" UE_AUTORTFM_NOAUTORTFM UE_AUTORTFM_API void autortfm_pre_static_local
 	}
 }
 
-extern "C" UE_AUTORTFM_NOAUTORTFM UE_AUTORTFM_API void autortfm_post_static_local_initializer()
+extern "C" UE_AUTORTFM_NOAUTORTFM UE_AUTORTFM_API void autortfm_post_static_local_initializer() noexcept
 {
 	if (FContext* Context = FContext::Get())
 	{
@@ -762,17 +752,13 @@ AUTORTFM_INTERNAL static void autortfm_open_from_closed_no_sanitize(void (*Work)
 {
 	FContext* const Context = FContext::Get();
 
-	AUTORTFM_SANITIZER_ONLY(Sanitizer::BeginDisable());
+	AUTORTFM_SANITIZER_ONLY(Sanitizer::FDisableHandle DisableHandle = Sanitizer::BeginDisable());
 
 	Work(Arg);
 
-	Context->MaybeThrow();
+	AUTORTFM_SANITIZER_ONLY(Sanitizer::EndDisable(DisableHandle));
 
-	// Note: This must come *after* MaybeThrow(), because we'll only throw
-	// if the current transaction at autortfm_pre_open() was aborted and
-	// popped, in which case the current transaction at autortfm_post_open()
-	// will now be an ancestor transaction.
-	AUTORTFM_SANITIZER_ONLY(Sanitizer::EndDisable());
+	Context->MaybeThrow();
 }
 
 extern "C" void autortfm_open_no_sanitize(void (*Work)(void*), void* Arg)
diff --git a/Engine/Source/Runtime/AutoRTFM/Private/Allocator.h b/Engine/Source/Runtime/AutoRTFM/Private/Allocator.h
--- a/Engine/Source/Runtime/AutoRTFM/Private/Allocator.h
+++ b/Engine/Source/Runtime/AutoRTFM/Private/Allocator.h
@@ -6,6 +6,7 @@
 
 #include "ExternAPI.h"
 #include "MemoryStats.h"
+#include "Utils.h"
 
 #include <concepts>
 #include <utility>
@@ -48,7 +49,7 @@ concept Allocator = requires {
 
 // The AutoRTFM default Allocator that wraps the external API's memory
 // allocation functions.
-struct AUTORTFM_DISABLE_NO_SANITIZE FAllocator
+struct AUTORTFM_INTERNAL FAllocator
 {
 	static inline void* Allocate(size_t Size, size_t Alignment)
 	{
diff --git a/Engine/Source/Runtime/AutoRTFM/Private/BlockAllocator.h b/Engine/Source/Runtime/AutoRTFM/Private/BlockAllocator.h
--- a/Engine/Source/Runtime/AutoRTFM/Private/BlockAllocator.h
+++ b/Engine/Source/Runtime/AutoRTFM/Private/BlockAllocator.h
@@ -22,7 +22,7 @@ namespace AutoRTFM
 //   DataAlignment - The maximum supported alignment of data from this allocator.
 //   GrowthPercentage - The percentage to grow the block size each reallocation.
 template <size_t InlineBlockDataSize = 256, size_t DataAlignment = 16, size_t GrowthPercentage = 200>
-class AUTORTFM_DISABLE_NO_SANITIZE TBlockAllocator
+class AUTORTFM_INTERNAL TBlockAllocator
 {
 public:
 	// Constructor
diff --git a/Engine/Source/Runtime/AutoRTFM/Private/CallstackTree.h b/Engine/Source/Runtime/AutoRTFM/Private/CallstackTree.h
--- a/Engine/Source/Runtime/AutoRTFM/Private/CallstackTree.h
+++ b/Engine/Source/Runtime/AutoRTFM/Private/CallstackTree.h
@@ -6,7 +6,6 @@
 
 #include "Allocator.h"
 #include "AutoRTFM/CAPI.h"
-#include "AutoRTFM/Defines.h"
 #include "Stack.h"
 #include "StlAllocator.h"
 #include "Utils.h"
@@ -20,7 +19,7 @@ namespace AutoRTFM
 // FCallstackTree is a store for callstacks that internally holds a DAG of
 // frames which reduces memory consumption by de-duplicating common base frames.
 template <Allocator Allocator = FAllocator>
-class AUTORTFM_DISABLE_NO_SANITIZE TCallstackTree final
+class AUTORTFM_INTERNAL TCallstackTree final
 {
 public:
 	using FHandle = uint32_t;
diff --git a/Engine/Source/Runtime/AutoRTFM/Private/CompilerInterface.cpp b/Engine/Source/Runtime/AutoRTFM/Private/CompilerInterface.cpp
--- a/Engine/Source/Runtime/AutoRTFM/Private/CompilerInterface.cpp
+++ b/Engine/Source/Runtime/AutoRTFM/Private/CompilerInterface.cpp
@@ -1,5 +1,6 @@
 // Copyright Epic Games, Inc. All Rights Reserved.
 
+#include "AutoRTFM/CAPI.h"
 #if (defined(__AUTORTFM) && __AUTORTFM)
 
 #include "AutoRTFM.h"
@@ -34,42 +35,42 @@ namespace AutoRTFM
 
 #endif
 
-extern "C" AUTORTFM_DISABLE_NO_SANITIZE UE_AUTORTFM_API AUTORTFM_NO_ASAN void autortfm_record_write(void* Ptr, size_t Size) noexcept
+extern "C" AUTORTFM_INTERNAL UE_AUTORTFM_API AUTORTFM_NO_ASAN void autortfm_record_write(void* Ptr, size_t Size) noexcept
 {
 	UE_AUTORTFM_HANDLE_NULL_WRITE(Ptr);
 	FContext* Context = FContext::Get();
 	Context->RecordWrite(Ptr, Size);
 }
 
-extern "C" AUTORTFM_DISABLE_NO_SANITIZE UE_AUTORTFM_API AUTORTFM_NO_ASAN void autortfm_record_write_1(void* Ptr) noexcept
+extern "C" AUTORTFM_INTERNAL UE_AUTORTFM_API AUTORTFM_NO_ASAN void autortfm_record_write_1(void* Ptr) noexcept
 {
 	UE_AUTORTFM_HANDLE_NULL_WRITE(Ptr);
 	FContext* Context = FContext::Get();
 	Context->RecordWrite<1>(Ptr);
 }
 
-extern "C" AUTORTFM_DISABLE_NO_SANITIZE UE_AUTORTFM_API AUTORTFM_NO_ASAN void autortfm_record_write_2(void* Ptr) noexcept
+extern "C" AUTORTFM_INTERNAL UE_AUTORTFM_API AUTORTFM_NO_ASAN void autortfm_record_write_2(void* Ptr) noexcept
 {
 	UE_AUTORTFM_HANDLE_NULL_WRITE(Ptr);
 	FContext* Context = FContext::Get();
 	Context->RecordWrite<2>(Ptr);
 }
 
-extern "C" AUTORTFM_DISABLE_NO_SANITIZE UE_AUTORTFM_API AUTORTFM_NO_ASAN void autortfm_record_write_4(void* Ptr) noexcept
+extern "C" AUTORTFM_INTERNAL UE_AUTORTFM_API AUTORTFM_NO_ASAN void autortfm_record_write_4(void* Ptr) noexcept
 {
 	UE_AUTORTFM_HANDLE_NULL_WRITE(Ptr);
 	FContext* Context = FContext::Get();
 	Context->RecordWrite<4>(Ptr);
 }
 
-extern "C" AUTORTFM_DISABLE_NO_SANITIZE UE_AUTORTFM_API AUTORTFM_NO_ASAN void autortfm_record_write_8(void* Ptr) noexcept
+extern "C" AUTORTFM_INTERNAL UE_AUTORTFM_API AUTORTFM_NO_ASAN void autortfm_record_write_8(void* Ptr) noexcept
 {
 	UE_AUTORTFM_HANDLE_NULL_WRITE(Ptr);
 	FContext* Context = FContext::Get();
 	Context->RecordWrite<8>(Ptr);
 }
 
-extern "C" AUTORTFM_DISABLE_NO_SANITIZE UE_AUTORTFM_API AUTORTFM_NO_ASAN void autortfm_record_masked_write(
+extern "C" AUTORTFM_INTERNAL UE_AUTORTFM_API AUTORTFM_NO_ASAN void autortfm_record_masked_write(
 	void* Ptr, uintptr_t Mask, int NumMaskBits, int ValueSizeBytes) noexcept
 {
 	UE_AUTORTFM_HANDLE_NULL_WRITE(Ptr);
@@ -106,14 +107,14 @@ static bool AutoRTFMIsPointerNegative(const void* InPointer)
 	return Address < 0;
 }
 
-extern "C" AUTORTFM_DISABLE_NO_SANITIZE UE_AUTORTFM_API AUTORTFM_NO_ASAN void autortfm_redirected_load(
+extern "C" AUTORTFM_INTERNAL UE_AUTORTFM_API AUTORTFM_NO_ASAN void autortfm_redirected_load(
 	uint32 AddressSpace, void* DestPointer, uint64 Size, uint64 SourceAddress, bool bWillWriteHint) noexcept
 {
 	ForTheRuntime::RedirectedLoad(AddressSpace ? AddressSpace : AutoRTFMExtractAddressSpaceFromAddress(SourceAddress), DestPointer,
 		Size, SourceAddress, bWillWriteHint);
 }
 
-extern "C" AUTORTFM_DISABLE_NO_SANITIZE UE_AUTORTFM_API AUTORTFM_NO_ASAN uint64 autortfm_redirected_load_8(
+extern "C" AUTORTFM_INTERNAL UE_AUTORTFM_API AUTORTFM_NO_ASAN uint64 autortfm_redirected_load_8(
 	uint32 AddressSpace, uint64 SourceAddress, bool bWillWriteHint) noexcept
 {
 	uint64 Result = 0;
@@ -122,7 +123,7 @@ extern "C" AUTORTFM_DISABLE_NO_SANITIZE UE_AUTORTFM_API AUTORTFM_NO_ASAN uint64
 	return Result;
 }
 
-extern "C" AUTORTFM_DISABLE_NO_SANITIZE UE_AUTORTFM_API AUTORTFM_NO_ASAN uint32 autortfm_redirected_load_4(
+extern "C" AUTORTFM_INTERNAL UE_AUTORTFM_API AUTORTFM_NO_ASAN uint32 autortfm_redirected_load_4(
 	uint32 AddressSpace, uint64 SourceAddress, bool bWillWriteHint) noexcept
 {
 	uint32 Result = 0;
@@ -131,7 +132,7 @@ extern "C" AUTORTFM_DISABLE_NO_SANITIZE UE_AUTORTFM_API AUTORTFM_NO_ASAN uint32
 	return Result;
 }
 
-extern "C" AUTORTFM_DISABLE_NO_SANITIZE UE_AUTORTFM_API AUTORTFM_NO_ASAN uint16 autortfm_redirected_load_2(
+extern "C" AUTORTFM_INTERNAL UE_AUTORTFM_API AUTORTFM_NO_ASAN uint16 autortfm_redirected_load_2(
 	uint32 AddressSpace, uint64 SourceAddress, bool bWillWriteHint) noexcept
 {
 	uint16 Result = 0;
@@ -140,7 +141,7 @@ extern "C" AUTORTFM_DISABLE_NO_SANITIZE UE_AUTORTFM_API AUTORTFM_NO_ASAN uint16
 	return Result;
 }
 
-extern "C" AUTORTFM_DISABLE_NO_SANITIZE UE_AUTORTFM_API AUTORTFM_NO_ASAN uint8 autortfm_redirected_load_1(
+extern "C" AUTORTFM_INTERNAL UE_AUTORTFM_API AUTORTFM_NO_ASAN uint8 autortfm_redirected_load_1(
 	uint32 AddressSpace, uint64 SourceAddress, bool bWillWriteHint) noexcept
 {
 	uint8 Result = 0;
@@ -149,35 +150,35 @@ extern "C" AUTORTFM_DISABLE_NO_SANITIZE UE_AUTORTFM_API AUTORTFM_NO_ASAN uint8 a
 	return Result;
 }
 
-extern "C" AUTORTFM_DISABLE_NO_SANITIZE UE_AUTORTFM_API AUTORTFM_NO_ASAN void autortfm_redirected_store(
+extern "C" AUTORTFM_INTERNAL UE_AUTORTFM_API AUTORTFM_NO_ASAN void autortfm_redirected_store(
 	uint32 AddressSpace, uint64 DestAddress, uint64 Size, const void* SourcePointer) noexcept
 {
 	ForTheRuntime::RedirectedStore(
 		AddressSpace ? AddressSpace : AutoRTFMExtractAddressSpaceFromAddress(DestAddress), DestAddress, Size, SourcePointer);
 }
 
-extern "C" AUTORTFM_DISABLE_NO_SANITIZE UE_AUTORTFM_API AUTORTFM_NO_ASAN void autortfm_redirected_store_8(
+extern "C" AUTORTFM_INTERNAL UE_AUTORTFM_API AUTORTFM_NO_ASAN void autortfm_redirected_store_8(
 	uint32 AddressSpace, uint64 DestAddress, uint64 Value) noexcept
 {
 	ForTheRuntime::RedirectedStore8(
 		AddressSpace ? AddressSpace : AutoRTFMExtractAddressSpaceFromAddress(DestAddress), DestAddress, &Value);
 }
 
-extern "C" AUTORTFM_DISABLE_NO_SANITIZE UE_AUTORTFM_API AUTORTFM_NO_ASAN void autortfm_redirected_store_4(
+extern "C" AUTORTFM_INTERNAL UE_AUTORTFM_API AUTORTFM_NO_ASAN void autortfm_redirected_store_4(
 	uint32 AddressSpace, uint64 DestAddress, uint32 Value) noexcept
 {
 	ForTheRuntime::RedirectedStore4(
 		AddressSpace ? AddressSpace : AutoRTFMExtractAddressSpaceFromAddress(DestAddress), DestAddress, &Value);
 }
 
-extern "C" AUTORTFM_DISABLE_NO_SANITIZE UE_AUTORTFM_API AUTORTFM_NO_ASAN void autortfm_redirected_store_2(
+extern "C" AUTORTFM_INTERNAL UE_AUTORTFM_API AUTORTFM_NO_ASAN void autortfm_redirected_store_2(
 	uint32 AddressSpace, uint64 DestAddress, uint16 Value) noexcept
 {
 	ForTheRuntime::RedirectedStore2(
 		AddressSpace ? AddressSpace : AutoRTFMExtractAddressSpaceFromAddress(DestAddress), DestAddress, &Value);
 }
 
-extern "C" AUTORTFM_DISABLE_NO_SANITIZE UE_AUTORTFM_API AUTORTFM_NO_ASAN void autortfm_redirected_store_1(
+extern "C" AUTORTFM_INTERNAL UE_AUTORTFM_API AUTORTFM_NO_ASAN void autortfm_redirected_store_1(
 	uint32 AddressSpace, uint64 DestAddress, uint8 Value) noexcept
 {
 	ForTheRuntime::RedirectedStore1(
@@ -191,26 +192,26 @@ extern "C" void* autortfm_lookup_function(void* OriginalFunction, const char* Wh
 	return FunctionMapLookup(OriginalFunction, Where);
 }
 
-extern "C" AUTORTFM_DISABLE_NO_SANITIZE UE_AUTORTFM_API void autortfm_memcpy(void* Dst, const void* Src, size_t Size) noexcept
+extern "C" AUTORTFM_INTERNAL UE_AUTORTFM_API void autortfm_memcpy(void* Dst, const void* Src, size_t Size) noexcept
 {
 	FContext* Context = FContext::Get();
 	Memcpy(Dst, Src, Size, Context);
 }
 
-extern "C" AUTORTFM_DISABLE_NO_SANITIZE UE_AUTORTFM_API void autortfm_memmove(void* Dst, const void* Src, size_t Size) noexcept
+extern "C" AUTORTFM_INTERNAL UE_AUTORTFM_API void autortfm_memmove(void* Dst, const void* Src, size_t Size) noexcept
 {
 	FContext* Context = FContext::Get();
 	Memmove(Dst, Src, Size, Context);
 }
 
-extern "C" AUTORTFM_DISABLE_NO_SANITIZE UE_AUTORTFM_API void autortfm_memset(void* Dst, int Value, size_t Size) noexcept
+extern "C" AUTORTFM_INTERNAL UE_AUTORTFM_API void autortfm_memset(void* Dst, int Value, size_t Size) noexcept
 {
 	FContext* Context = FContext::Get();
 	Memset(Dst, Value, Size, Context);
 }
 
 template <bool bDstIsRedirected, bool bSrcIsRedirected>
-AUTORTFM_DISABLE_NO_SANITIZE void autortfm_memmove_redirectable_impl(void* Dst, const void* Src, size_t Size) noexcept
+AUTORTFM_INTERNAL void autortfm_memmove_redirectable_impl(void* Dst, const void* Src, size_t Size) noexcept
 {
 	int8 MoveBackwards = (Src < Dst);
 
@@ -302,7 +303,7 @@ AUTORTFM_DISABLE_NO_SANITIZE void autortfm_memmove_redirectable_impl(void* Dst,
 }
 
 template <bool bDstIsRedirected>
-AUTORTFM_DISABLE_NO_SANITIZE void autortfm_memset_redirectable_impl(void* Dst, uint8 Value, size_t Size) noexcept
+AUTORTFM_INTERNAL void autortfm_memset_redirectable_impl(void* Dst, uint8 Value, size_t Size) noexcept
 {
 	uint8* DstCursor = (uint8*)Dst;
 
@@ -342,8 +343,7 @@ AUTORTFM_DISABLE_NO_SANITIZE void autortfm_memset_redirectable_impl(void* Dst, u
 	}
 }
 
-extern "C" AUTORTFM_DISABLE_NO_SANITIZE UE_AUTORTFM_API void autortfm_memmove_redirectable(
-	void* Dst, const void* Src, size_t Size) noexcept
+extern "C" AUTORTFM_INTERNAL UE_AUTORTFM_API void autortfm_memmove_redirectable(void* Dst, const void* Src, size_t Size) noexcept
 {
 	FContext* Context = FContext::Get();
 
@@ -376,13 +376,12 @@ extern "C" AUTORTFM_DISABLE_NO_SANITIZE UE_AUTORTFM_API void autortfm_memmove_re
 	}
 }
 
-extern "C" AUTORTFM_DISABLE_NO_SANITIZE UE_AUTORTFM_API void autortfm_memcpy_redirectable(
-	void* Dst, const void* Src, size_t Size) noexcept
+extern "C" AUTORTFM_INTERNAL UE_AUTORTFM_API void autortfm_memcpy_redirectable(void* Dst, const void* Src, size_t Size) noexcept
 {
 	autortfm_memmove_redirectable(Dst, Src, Size);
 }
 
-extern "C" AUTORTFM_DISABLE_NO_SANITIZE UE_AUTORTFM_API void autortfm_memset_redirectable(void* Dst, int Value, size_t Size) noexcept
+extern "C" AUTORTFM_INTERNAL UE_AUTORTFM_API void autortfm_memset_redirectable(void* Dst, int Value, size_t Size) noexcept
 {
 	FContext* Context = FContext::Get();
 
@@ -405,7 +404,7 @@ extern "C" void autortfm_unreachable(const char* Message) noexcept
 	__builtin_unreachable();
 }
 
-extern "C" AUTORTFM_DISABLE_NO_SANITIZE UE_AUTORTFM_API void autortfm_llvm_fail(const char* Message) noexcept
+extern "C" AUTORTFM_INTERNAL UE_AUTORTFM_API void autortfm_llvm_fail(const char* Message) noexcept
 {
 	if (Message)
 	{
@@ -417,7 +416,7 @@ extern "C" AUTORTFM_DISABLE_NO_SANITIZE UE_AUTORTFM_API void autortfm_llvm_fail(
 	}
 }
 
-extern "C" AUTORTFM_DISABLE_NO_SANITIZE UE_AUTORTFM_API void autortfm_llvm_missing_function() noexcept
+extern "C" AUTORTFM_INTERNAL UE_AUTORTFM_API void autortfm_llvm_missing_function() noexcept
 {
 	if (ForTheRuntime::GetInternalAbortAction() == ForTheRuntime::EAutoRTFMInternalAbortActionState::Crash)
 	{
@@ -432,16 +431,11 @@ extern "C" AUTORTFM_DISABLE_NO_SANITIZE UE_AUTORTFM_API void autortfm_llvm_missi
 	Context->AbortTransactionAndThrow(ETransactionStatus::AbortedByLanguage);
 }
 
-extern "C" AUTORTFM_DISABLE_NO_SANITIZE UE_AUTORTFM_FORCENOINLINE UE_AUTORTFM_API void autortfm_called_no_autortfm() noexcept
+extern "C" AUTORTFM_INTERNAL UE_AUTORTFM_FORCENOINLINE UE_AUTORTFM_API void autortfm_called_no_autortfm() noexcept
 {
 	AUTORTFM_FATAL("inlined UE_AUTORTFM_NOAUTORTFM function called from the closed");
 }
 
-extern "C" AUTORTFM_DISABLE_NO_SANITIZE UE_AUTORTFM_FORCENOINLINE UE_AUTORTFM_API void autortfm_disable_called() noexcept
-{
-	AUTORTFM_FATAL("inlined AUTORTFM_DISABLE function called from the closed");
-}
-
 }  // namespace AutoRTFM
 
 #endif  // defined(__AUTORTFM) && __AUTORTFM
diff --git a/Engine/Source/Runtime/AutoRTFM/Private/Context.h b/Engine/Source/Runtime/AutoRTFM/Private/Context.h
--- a/Engine/Source/Runtime/AutoRTFM/Private/Context.h
+++ b/Engine/Source/Runtime/AutoRTFM/Private/Context.h
@@ -18,7 +18,7 @@
 namespace AutoRTFM
 {
 
-class AUTORTFM_DISABLE_NO_SANITIZE FContext final
+class AUTORTFM_INTERNAL FContext final
 {
 	using FTaskPool = typename TTaskArray<TTask<void()>>::FEntryPool;
 	using FTransactionPool = TIntrusivePool<FTransaction, 16>;
diff --git a/Engine/Source/Runtime/AutoRTFM/Private/FunctionMap.cpp b/Engine/Source/Runtime/AutoRTFM/Private/FunctionMap.cpp
--- a/Engine/Source/Runtime/AutoRTFM/Private/FunctionMap.cpp
+++ b/Engine/Source/Runtime/AutoRTFM/Private/FunctionMap.cpp
@@ -258,7 +258,7 @@ inline void* FollowRelocation(void* Function)
 	return nullptr;
 }
 
-AUTORTFM_DISABLE_NO_SANITIZE static void* FunctionMapReportError(void* Open, const char* Where)
+AUTORTFM_INTERNAL static void* FunctionMapReportError(void* Open, const char* Where)
 {
 #if AUTORTFM_PLATFORM_WINDOWS
 	if (Open == &InitializeCriticalSection || Open == &InitializeCriticalSectionAndSpinCount || Open == &SetCriticalSectionSpinCount
diff --git a/Engine/Source/Runtime/AutoRTFM/Private/FunctionMap.h b/Engine/Source/Runtime/AutoRTFM/Private/FunctionMap.h
--- a/Engine/Source/Runtime/AutoRTFM/Private/FunctionMap.h
+++ b/Engine/Source/Runtime/AutoRTFM/Private/FunctionMap.h
@@ -4,23 +4,23 @@
 
 #if (defined(__AUTORTFM) && __AUTORTFM)
 
-#include "AutoRTFM.h"
+#include "Utils.h"
 
 namespace AutoRTFM
 {
 
 // Dumps statistics about the function map to LOG_INFO.
-AUTORTFM_DISABLE_NO_SANITIZE void FunctionMapDumpStats();
+AUTORTFM_INTERNAL void FunctionMapDumpStats();
 
 // Adds all the open to closed function mappings from the linked list of tables to the function map.
-AUTORTFM_DISABLE_NO_SANITIZE void FunctionMapAdd(autortfm_open_to_closed_table* Tables);
+AUTORTFM_INTERNAL void FunctionMapAdd(autortfm_open_to_closed_table* Tables);
 
 // Looks up the closed function from the open function using the function map, following dynamic
 // library thunk redirections.
-AUTORTFM_DISABLE_NO_SANITIZE void* FunctionMapLookupExhaustive(void* OpenFn, const char* Where);
+AUTORTFM_INTERNAL void* FunctionMapLookupExhaustive(void* OpenFn, const char* Where);
 
 template <typename ReturnType, typename... ParameterTypes>
-AUTORTFM_DISABLE_NO_SANITIZE auto FunctionMapLookupExhaustive(ReturnType (*Function)(ParameterTypes...), const char* Where)
+AUTORTFM_INTERNAL auto FunctionMapLookupExhaustive(ReturnType (*Function)(ParameterTypes...), const char* Where)
 	-> ReturnType (*)(ParameterTypes...)
 {
 	return reinterpret_cast<ReturnType (*)(ParameterTypes...)>(
diff --git a/Engine/Source/Runtime/AutoRTFM/Private/FunctionMapInlines.h b/Engine/Source/Runtime/AutoRTFM/Private/FunctionMapInlines.h
--- a/Engine/Source/Runtime/AutoRTFM/Private/FunctionMapInlines.h
+++ b/Engine/Source/Runtime/AutoRTFM/Private/FunctionMapInlines.h
@@ -12,7 +12,7 @@
 namespace AutoRTFM
 {
 
-AUTORTFM_DISABLE_NO_SANITIZE UE_AUTORTFM_FORCEINLINE void* FunctionMapLookupUsingMagicPrefix(void* OpenFn)
+AUTORTFM_INTERNAL UE_AUTORTFM_FORCEINLINE void* FunctionMapLookupUsingMagicPrefix(void* OpenFn)
 {
 	constexpr uint64_t UbsanMagic = 0xc105'cafe;
 
@@ -76,7 +76,7 @@ AUTORTFM_DISABLE_NO_SANITIZE UE_AUTORTFM_FORCEINLINE void* FunctionMapLookupUsin
 	return nullptr;
 }
 
-AUTORTFM_DISABLE_NO_SANITIZE inline void* FunctionMapLookup(void* OpenFn, const char* Where)
+AUTORTFM_INTERNAL inline void* FunctionMapLookup(void* OpenFn, const char* Where)
 {
 	if (void* ClosedFn = FunctionMapLookupUsingMagicPrefix(OpenFn); AUTORTFM_LIKELY(ClosedFn))
 	{
@@ -87,7 +87,7 @@ AUTORTFM_DISABLE_NO_SANITIZE inline void* FunctionMapLookup(void* OpenFn, const
 }
 
 template <typename TReturnType, typename... TParameterTypes>
-AUTORTFM_DISABLE_NO_SANITIZE auto FunctionMapLookup(TReturnType (*OpenFn)(TParameterTypes...), const char* Where)
+AUTORTFM_INTERNAL auto FunctionMapLookup(TReturnType (*OpenFn)(TParameterTypes...), const char* Where)
 	-> TReturnType (*)(TParameterTypes...)
 {
 	return reinterpret_cast<TReturnType (*)(TParameterTypes...)>(FunctionMapLookup(reinterpret_cast<void*>(OpenFn), Where));
diff --git a/Engine/Source/Runtime/AutoRTFM/Private/IntervalTree.h b/Engine/Source/Runtime/AutoRTFM/Private/IntervalTree.h
--- a/Engine/Source/Runtime/AutoRTFM/Private/IntervalTree.h
+++ b/Engine/Source/Runtime/AutoRTFM/Private/IntervalTree.h
@@ -22,7 +22,7 @@ namespace AutoRTFM
 //   - Overlapping intervals with different data are trimmed, split, or deleted.
 template <typename DataTypeIn, Allocator Allocator = FAllocator>
 requires std::equality_comparable<DataTypeIn>
-struct AUTORTFM_DISABLE_NO_SANITIZE TIntervalTree final
+struct AUTORTFM_INTERNAL TIntervalTree final
 {
 private:
 	using NodeIndex = uint32_t;
diff --git a/Engine/Source/Runtime/AutoRTFM/Private/IntrusivePool.h b/Engine/Source/Runtime/AutoRTFM/Private/IntrusivePool.h
--- a/Engine/Source/Runtime/AutoRTFM/Private/IntrusivePool.h
+++ b/Engine/Source/Runtime/AutoRTFM/Private/IntrusivePool.h
@@ -14,7 +14,7 @@ namespace AutoRTFM
 {
 
 template <typename ItemType, size_t InlineCapacity>
-class AUTORTFM_DISABLE_NO_SANITIZE TIntrusivePool final
+class AUTORTFM_INTERNAL TIntrusivePool final
 {
 public:
 	using FItem = ItemType;
diff --git a/Engine/Source/Runtime/AutoRTFM/Private/Memcpy.h b/Engine/Source/Runtime/AutoRTFM/Private/Memcpy.h
--- a/Engine/Source/Runtime/AutoRTFM/Private/Memcpy.h
+++ b/Engine/Source/Runtime/AutoRTFM/Private/Memcpy.h
@@ -4,7 +4,7 @@
 
 #if (defined(__AUTORTFM) && __AUTORTFM)
 
-#include "AutoRTFM/Defines.h"
+#include "Utils.h"
 
 #include <stdlib.h>
 
@@ -13,10 +13,10 @@ namespace AutoRTFM
 
 class FContext;
 
-AUTORTFM_DISABLE_NO_SANITIZE void* MemcpyToNew(void* Dst, const void* Src, size_t Size, FContext* Context);
-AUTORTFM_DISABLE_NO_SANITIZE void* Memcpy(void* Dst, const void* Src, size_t Size, FContext* Context);
-AUTORTFM_DISABLE_NO_SANITIZE void* Memmove(void* Dst, const void* Src, size_t Size, FContext* Context);
-AUTORTFM_DISABLE_NO_SANITIZE void* Memset(void* Dst, int Value, size_t Size, FContext* Context);
+AUTORTFM_INTERNAL void* MemcpyToNew(void* Dst, const void* Src, size_t Size, FContext* Context);
+AUTORTFM_INTERNAL void* Memcpy(void* Dst, const void* Src, size_t Size, FContext* Context);
+AUTORTFM_INTERNAL void* Memmove(void* Dst, const void* Src, size_t Size, FContext* Context);
+AUTORTFM_INTERNAL void* Memset(void* Dst, int Value, size_t Size, FContext* Context);
 
 }  // namespace AutoRTFM
 
diff --git a/Engine/Source/Runtime/AutoRTFM/Private/MemoryTracker.h b/Engine/Source/Runtime/AutoRTFM/Private/MemoryTracker.h
--- a/Engine/Source/Runtime/AutoRTFM/Private/MemoryTracker.h
+++ b/Engine/Source/Runtime/AutoRTFM/Private/MemoryTracker.h
@@ -22,7 +22,7 @@ namespace AutoRTFM
 // Assumes that queries never partially overlap a stored interval — a query either falls
 // entirely within an interval or entirely outside it. This assumption enables a simplified
 // BST traversal in Contains() that avoids examining both subtrees.
-struct AUTORTFM_DISABLE_NO_SANITIZE FMemoryTracker final
+struct AUTORTFM_INTERNAL FMemoryTracker final
 {
 private:
 	using NodeIndex = uint32_t;
diff --git a/Engine/Source/Runtime/AutoRTFM/Private/Pool.h b/Engine/Source/Runtime/AutoRTFM/Private/Pool.h
--- a/Engine/Source/Runtime/AutoRTFM/Private/Pool.h
+++ b/Engine/Source/Runtime/AutoRTFM/Private/Pool.h
@@ -23,7 +23,7 @@ namespace AutoRTFM
 //   InlineCapacity - the number of items that can be held before spilling to
 //                    the heap.
 template <typename ItemType, size_t InlineCapacity>
-class AUTORTFM_DISABLE_NO_SANITIZE TPool
+class AUTORTFM_INTERNAL TPool
 {
 public:
 	using FItem = ItemType;
diff --git a/Engine/Source/Runtime/AutoRTFM/Private/Sanitizer.cpp b/Engine/Source/Runtime/AutoRTFM/Private/Sanitizer.cpp
--- a/Engine/Source/Runtime/AutoRTFM/Private/Sanitizer.cpp
+++ b/Engine/Source/Runtime/AutoRTFM/Private/Sanitizer.cpp
@@ -14,13 +14,12 @@
 #include "StackRange.h"
 #include "Utils.h"
 
-#include <memory>
 #include <mutex>
 
 namespace AutoRTFM::Sanitizer
 {
 
-struct AUTORTFM_DISABLE_NO_SANITIZE FOSAllocator
+struct AUTORTFM_INTERNAL FOSAllocator
 {
 	static void* Allocate(size_t Size, size_t Alignment)
 	{
@@ -57,7 +56,7 @@ static_assert(Allocator<FOSAllocator>);
 // If more are found, the report will note how many were omitted.
 static constexpr size_t MaxReportedClosedWrites = 4;
 
-struct AUTORTFM_DISABLE_NO_SANITIZE FSanitizer
+struct AUTORTFM_INTERNAL FSanitizer
 {
 	using FCallstackTree = TCallstackTree<FOSAllocator>;
 
@@ -92,8 +91,8 @@ struct AUTORTFM_DISABLE_NO_SANITIZE FSanitizer
 	static void Initialize();
 	static void Shutdown();
 
-	void DisableTransaction();
-	void ReenableTransaction();
+	FDisableHandle DisableTransaction();
+	void ReenableTransaction(FDisableHandle);
 
 	size_t NumberOfIssuesFound() const;
 
@@ -107,7 +106,7 @@ struct AUTORTFM_DISABLE_NO_SANITIZE FSanitizer
 private:
 	using TClosedWriteTree = TIntervalTree<FCallstackTree::FHandle, FOSAllocator>;
 
-	struct AUTORTFM_DISABLE_NO_SANITIZE FTransaction
+	struct FTransaction
 	{
 		TransactionID const ID;
 		FTransaction* const Parent;
@@ -133,7 +132,7 @@ std::atomic<Sanitizer::EMode> FSanitizer::Mode{EMode::Error};
 std::atomic<size_t> FSanitizer::NumTransactionsRunning{0};
 std::atomic<bool> FSanitizer::bRecordClosedWriteCallstacks{true};
 
-struct AUTORTFM_DISABLE_NO_SANITIZE FSanitizer::FLock
+struct AUTORTFM_INTERNAL FSanitizer::FLock
 {
 	FLock()
 	{
@@ -193,20 +192,25 @@ void FSanitizer::Shutdown()
 	Instance = nullptr;
 }
 
-void FSanitizer::DisableTransaction()
+FDisableHandle FSanitizer::DisableTransaction()
 {
-	if (FTransaction* Transaction = CurrentTransaction)
+	if (CurrentTransaction)
 	{
-		Transaction->DisableCount++;
+		CurrentTransaction->DisableCount++;
 	}
+	return CurrentTransaction;
 }
 
-void FSanitizer::ReenableTransaction()
+void FSanitizer::ReenableTransaction(FDisableHandle Handle)
 {
-	if (FTransaction* Transaction = CurrentTransaction)
+	for (FTransaction* Transaction = CurrentTransaction; Transaction; Transaction = Transaction->Parent)
 	{
-		AUTORTFM_ASSERT(Transaction->DisableCount > 0);
-		Transaction->DisableCount--;
+		if (Handle == Transaction)
+		{
+			AUTORTFM_ASSERT(Transaction->DisableCount > 0);
+			Transaction->DisableCount--;
+			break;
+		}
 	}
 }
 
@@ -266,7 +270,7 @@ void FSanitizer::AbortTransaction(TransactionID ID)
 		return nullptr;
 	}
 
-	// Collision found — gather all closed writes that overlap the open write range.
+	// Collision found - gather all closed writes that overlap the open write range.
 	size_t NumClosedWrites = 0;
 	size_t TotalClosedWrites = 0;
 	FConflictingWrite::FClosedWriteEntry ClosedWrites[MaxReportedClosedWrites];
@@ -374,7 +378,7 @@ void FSanitizer::FConflictingWrite::Report() const
 	}
 	else if (TotalClosedWrites == 1)
 	{
-		// Single closed write — use the familiar format.
+		// Single closed write - use the familiar format.
 		const FClosedWriteEntry& Entry = ClosedWrites[0];
 		if (!Entry.Callstack.IsEmpty())
 		{
@@ -390,7 +394,7 @@ void FSanitizer::FConflictingWrite::Report() const
 	}
 	else
 	{
-		// Multiple closed writes — numbered list.
+		// Multiple closed writes - numbered list.
 		AUTORTFM_LOG_AT_SEVERITY(Severity, "AutoRTFMSan: %zu conflicting transactional write(s):", TotalClosedWrites);
 		for (size_t I = 0; I < NumClosedWrites; I++)
 		{
@@ -481,7 +485,7 @@ void ClosedWrite(TransactionID ID, void* Ptr, size_t Size, void* ProgramCounter)
 ////////////////////////////////////////////////////////////////////////////////
 // Compiler Interface
 ////////////////////////////////////////////////////////////////////////////////
-extern "C" AUTORTFM_DISABLE_NO_SANITIZE UE_AUTORTFM_FORCENOINLINE UE_AUTORTFM_API void autortfm_sanitizer_open_write(
+extern "C" AUTORTFM_INTERNAL UE_AUTORTFM_FORCENOINLINE UE_AUTORTFM_API void autortfm_sanitizer_open_write(
 	void* Ptr, size_t Size) noexcept
 {
 	// Skip the cost of a mutex lock if there are no transactions running, or
@@ -543,16 +547,29 @@ void SetRecordClosedWriteCallstacks(bool bEnable)
 	}
 }
 
-void BeginDisable()
+FDisableHandle BeginDisable()
 {
-	FSanitizer::FLock Sanitizer;
-	Sanitizer->DisableTransaction();
+	if (AnyTransactionsInFlight() && FSanitizer::Mode != EMode::Disabled)
+	{
+		FSanitizer::FLock Sanitizer;
+		if (Sanitizer)
+		{
+			return Sanitizer->DisableTransaction();
+		}
+	}
+	return nullptr;
 }
 
-void EndDisable()
+void EndDisable(FDisableHandle Handle)
 {
-	FSanitizer::FLock Sanitizer;
-	Sanitizer->ReenableTransaction();
+	if (Handle)
+	{
+		FSanitizer::FLock Sanitizer;
+		if (Sanitizer)
+		{
+			Sanitizer->ReenableTransaction(Handle);
+		}
+	}
 }
 
 bool AnyTransactionsInFlight()
diff --git a/Engine/Source/Runtime/AutoRTFM/Private/Sanitizer.h b/Engine/Source/Runtime/AutoRTFM/Private/Sanitizer.h
--- a/Engine/Source/Runtime/AutoRTFM/Private/Sanitizer.h
+++ b/Engine/Source/Runtime/AutoRTFM/Private/Sanitizer.h
@@ -11,16 +11,20 @@
 namespace AutoRTFM::Sanitizer
 {
 
-AUTORTFM_DISABLE_NO_SANITIZE void Initialize();
-AUTORTFM_DISABLE_NO_SANITIZE void Shutdown();
+#define AUTORTFM_SANITIZER_INTERNAL [[clang::autortfm(autortfm_mode_internal)]]
 
-AUTORTFM_DISABLE_NO_SANITIZE size_t NumberOfIssuesFound();
+AUTORTFM_SANITIZER_INTERNAL void Initialize();
+AUTORTFM_SANITIZER_INTERNAL void Shutdown();
 
-AUTORTFM_DISABLE_NO_SANITIZE void StartTransaction(TransactionID ID, FStackRange StackRange);
-AUTORTFM_DISABLE_NO_SANITIZE void CommitTransaction(TransactionID ID);
-AUTORTFM_DISABLE_NO_SANITIZE void AbortTransaction(TransactionID ID);
+AUTORTFM_SANITIZER_INTERNAL size_t NumberOfIssuesFound();
 
-AUTORTFM_DISABLE_NO_SANITIZE void ClosedWrite(TransactionID ID, void* Ptr, size_t Size, void* ProgramCounter);
+AUTORTFM_SANITIZER_INTERNAL void StartTransaction(TransactionID ID, FStackRange StackRange);
+AUTORTFM_SANITIZER_INTERNAL void CommitTransaction(TransactionID ID);
+AUTORTFM_SANITIZER_INTERNAL void AbortTransaction(TransactionID ID);
+
+AUTORTFM_SANITIZER_INTERNAL void ClosedWrite(TransactionID ID, void* Ptr, size_t Size, void* ProgramCounter);
+
+#undef AUTORTFM_SANITIZER_INTERNAL
 
 }  // namespace AutoRTFM::Sanitizer
 
diff --git a/Engine/Source/Runtime/AutoRTFM/Private/ScopedGuard.h b/Engine/Source/Runtime/AutoRTFM/Private/ScopedGuard.h
--- a/Engine/Source/Runtime/AutoRTFM/Private/ScopedGuard.h
+++ b/Engine/Source/Runtime/AutoRTFM/Private/ScopedGuard.h
@@ -4,13 +4,13 @@
 
 #if (defined(__AUTORTFM) && __AUTORTFM)
 
-#include "AutoRTFM/Defines.h"
+#include "Utils.h"
 
 namespace AutoRTFM
 {
 
 template <typename T>
-class AUTORTFM_DISABLE_NO_SANITIZE TScopedGuard
+class AUTORTFM_INTERNAL TScopedGuard
 {
 public:
 	TScopedGuard(T& Ref, const T& Value) : OldValue(Ref), Ref(Ref)
diff --git a/Engine/Source/Runtime/AutoRTFM/Private/Stack.h b/Engine/Source/Runtime/AutoRTFM/Private/Stack.h
--- a/Engine/Source/Runtime/AutoRTFM/Private/Stack.h
+++ b/Engine/Source/Runtime/AutoRTFM/Private/Stack.h
@@ -32,7 +32,7 @@ namespace AutoRTFM
 //   Validation - if Disabled, do not perform validity assertions.
 template <typename T, size_t InlineCapacity, Allocator Allocator = FAllocator,
 	EContainerValidation Validation = EContainerValidation::Enabled>
-class AUTORTFM_DISABLE_NO_SANITIZE TStack
+class AUTORTFM_INTERNAL TStack
 {
 public:
 	using ElementType = T;
diff --git a/Engine/Source/Runtime/AutoRTFM/Private/Transaction.h b/Engine/Source/Runtime/AutoRTFM/Private/Transaction.h
--- a/Engine/Source/Runtime/AutoRTFM/Private/Transaction.h
+++ b/Engine/Source/Runtime/AutoRTFM/Private/Transaction.h
@@ -17,7 +17,7 @@ namespace AutoRTFM
 {
 class FContext;
 
-class AUTORTFM_DISABLE_NO_SANITIZE FTransaction final
+class AUTORTFM_INTERNAL FTransaction final
 {
 	friend class TIntrusivePool<FTransaction, 16>;
 
diff --git a/Engine/Source/Runtime/AutoRTFM/Private/Utils.h b/Engine/Source/Runtime/AutoRTFM/Private/Utils.h
--- a/Engine/Source/Runtime/AutoRTFM/Private/Utils.h
+++ b/Engine/Source/Runtime/AutoRTFM/Private/Utils.h
@@ -2,10 +2,9 @@
 
 #pragma once
 
-#include "AutoRTFM/Defines.h"
 #if (defined(__AUTORTFM) && __AUTORTFM)
 
-#include "AutoRTFM.h"
+#include "AutoRTFM/Defines.h"
 #include "BuildMacros.h"
 #include "ExternAPI.h"
 
diff --git a/Engine/Source/Runtime/AutoRTFM/Private/WriteLog.h b/Engine/Source/Runtime/AutoRTFM/Private/WriteLog.h
--- a/Engine/Source/Runtime/AutoRTFM/Private/WriteLog.h
+++ b/Engine/Source/Runtime/AutoRTFM/Private/WriteLog.h
@@ -17,7 +17,7 @@
 
 namespace AutoRTFM
 {
-struct AUTORTFM_DISABLE_NO_SANITIZE FWriteLogEntry final
+struct AUTORTFM_INTERNAL FWriteLogEntry final
 {
 	// The address of the write.
 	std::byte* LogicalAddress = nullptr;
@@ -48,7 +48,7 @@ struct AUTORTFM_DISABLE_NO_SANITIZE FWriteLogEntry final
 // FWriteLog holds an ordered list of write records which can be iterated
 // forwards and backwards.
 // Ensure changes to this class are kept in sync with AutoRTFM.natvis.
-class AUTORTFM_DISABLE_NO_SANITIZE FWriteLog final
+class AUTORTFM_INTERNAL FWriteLog final
 {
 	enum EAppendMode
 	{
diff --git a/Engine/Source/Runtime/AutoRTFM/Public/AutoRTFM.h b/Engine/Source/Runtime/AutoRTFM/Public/AutoRTFM.h
--- a/Engine/Source/Runtime/AutoRTFM/Public/AutoRTFM.h
+++ b/Engine/Source/Runtime/AutoRTFM/Public/AutoRTFM.h
@@ -226,6 +226,9 @@ enum class EContextStatus : uint8_t
 using TransactionID = autortfm_transaction_id;
 
 #if AUTORTFM_SANITIZER
+
+#define AUTORTFM_SANITIZER_INTERNAL [[clang::autortfm(autortfm_mode_internal)]]
+
 // AutoRTFM sanitizer API (requires the -fsanitize=autortfm compiler flag)
 namespace Sanitizer
 {
@@ -246,16 +249,19 @@ enum class EMode : uint8_t
 };
 
 // Returns the global AutoRTFM sanitizer mode currently in use.
-AUTORTFM_OPEN UE_AUTORTFM_API EMode GetMode();
+AUTORTFM_SANITIZER_INTERNAL UE_AUTORTFM_API EMode GetMode();
 
 // Sets the global AutoRTFM sanitizer mode.
-AUTORTFM_OPEN UE_AUTORTFM_API void SetMode(EMode Mode);
+AUTORTFM_SANITIZER_INTERNAL UE_AUTORTFM_API void SetMode(EMode Mode);
 
 // Returns true if the AutoRTFM sanitizer records the callstack of closed writes.
-AUTORTFM_OPEN UE_AUTORTFM_API bool RecordClosedWriteCallstacks();
+AUTORTFM_SANITIZER_INTERNAL UE_AUTORTFM_API bool RecordClosedWriteCallstacks();
 
 // Sets whether the AutoRTFM sanitizer records the callstack of closed writes.
-AUTORTFM_OPEN UE_AUTORTFM_API void SetRecordClosedWriteCallstacks(bool bEnable);
+AUTORTFM_SANITIZER_INTERNAL UE_AUTORTFM_API void SetRecordClosedWriteCallstacks(bool bEnable);
+
+// Opaque handle returned by BeginDisable()
+using FDisableHandle = void*;
 
 // Temporarily disables the AutoRTFM sanitizer for the current transaction
 // and any child transactions.
@@ -263,44 +269,46 @@ AUTORTFM_OPEN UE_AUTORTFM_API void SetRecordClosedWriteCallstacks(bool bEnable);
 // Calls to BeginDisable() are internally counted and may be nested within
 // other BeginDisable() / EndDisable() calls.
 // This call is a no-op when called outside of a transaction.
-AUTORTFM_OPEN UE_AUTORTFM_API void BeginDisable();
+AUTORTFM_SANITIZER_INTERNAL UE_AUTORTFM_API FDisableHandle BeginDisable();
 
 // Ends temporarily disabling the AutoRTFM sanitizer.
+// DisableHandle must be a handle returned by BeginDisable().
 // See BeginDisable().
-AUTORTFM_OPEN UE_AUTORTFM_API void EndDisable();
+AUTORTFM_SANITIZER_INTERNAL UE_AUTORTFM_API void EndDisable(FDisableHandle);
 
 // Returns true if there are any AutoRTFM transactions in flight.
 // Can be safely called from any thread, although the result will be
 // non-deterministic without synchronization between threads.
-AUTORTFM_OPEN UE_AUTORTFM_API bool AnyTransactionsInFlight();
+AUTORTFM_SANITIZER_INTERNAL UE_AUTORTFM_API bool AnyTransactionsInFlight();
 
 // A helper class that temporarily disables the AutoRTFM sanitizer for
 // the lifetime of the object.
 class FDisableScope
 {
 public:
-	UE_AUTORTFM_FORCEINLINE FDisableScope()
-	{
-		BeginDisable();
-	}
+	UE_AUTORTFM_FORCEINLINE FDisableScope() : Handle{BeginDisable()} {}
 	UE_AUTORTFM_FORCEINLINE ~FDisableScope()
 	{
-		EndDisable();
+		EndDisable(Handle);
 	}
 
 private:
 	FDisableScope(const FDisableScope&) = delete;
 	FDisableScope& operator=(const FDisableScope&) = delete;
+	FDisableHandle const Handle;
 };
 }  // namespace Sanitizer
-#endif  // AUTORTFM_SANITIZER
 
-#if AUTORTFM_SANITIZER
+#undef AUTORTFM_SANITIZER_INTERNAL
+
 #define AUTORTFM_SANITIZER_DISABLE_SCOPE() \
 	::AutoRTFM::Sanitizer::FDisableScope UE_AUTORTFM_CONCAT(AutoRTFMSanitizerModeScope, __COUNTER__)
-#else
+
+#else  // ^^^ AUTORTFM_SANITIZER ^^^ | vvv !AUTORTFM_SANITIZER vvv
+
 #define AUTORTFM_SANITIZER_DISABLE_SCOPE() static_assert(true) /* require semicolon */
-#endif
+
+#endif  // ^^^ !AUTORTFM_SANITIZER ^^^
 
 #if UE_AUTORTFM
 namespace ForTheRuntime
